Fixing issues with ML when the sdk is not set.

diff --git a/Engine/Plugins/Lumin/MagicLeapCamera/Source/Private/MagicLeapCameraPlugin.cpp b/Engine/Plugins/Lumin/MagicLeapCamera/Source/Private/MagicLeapCameraPlugin.cpp
--- a/Engine/Plugins/Lumin/MagicLeapCamera/Source/Private/MagicLeapCameraPlugin.cpp
+++ b/Engine/Plugins/Lumin/MagicLeapCamera/Source/Private/MagicLeapCameraPlugin.cpp
@@ -285,3 +285,5 @@ bool FMagicLeapCameraPlugin::TryPushNewCaptureTask(FCameraTask::EType InTaskType
 
 	return false;
 }
+
+IMPLEMENT_MODULE(FMagicLeapCameraPlugin, MagicLeapCamera);
diff --git a/Engine/Plugins/Lumin/MagicLeapCamera/Source/Private/MagicLeapCameraPlugin.h b/Engine/Plugins/Lumin/MagicLeapCamera/Source/Private/MagicLeapCameraPlugin.h
--- a/Engine/Plugins/Lumin/MagicLeapCamera/Source/Private/MagicLeapCameraPlugin.h
+++ b/Engine/Plugins/Lumin/MagicLeapCamera/Source/Private/MagicLeapCameraPlugin.h
@@ -49,6 +49,4 @@ private:
 	FCameraLogMessageMulti OnLogMessage;
 };
 
-IMPLEMENT_MODULE(FMagicLeapCameraPlugin, MagicLeapCamera);
-
 #define GET_MAGIC_LEAP_CAMERA_PLUGIN() static_cast<FMagicLeapCameraPlugin*>(FModuleManager::Get().GetModule("MagicLeapCamera"))
diff --git a/Engine/Plugins/Lumin/MagicLeapCamera/Source/Private/MagicLeapCameraRunnable.cpp b/Engine/Plugins/Lumin/MagicLeapCamera/Source/Private/MagicLeapCameraRunnable.cpp
--- a/Engine/Plugins/Lumin/MagicLeapCamera/Source/Private/MagicLeapCameraRunnable.cpp
+++ b/Engine/Plugins/Lumin/MagicLeapCamera/Source/Private/MagicLeapCameraRunnable.cpp
@@ -13,7 +13,7 @@ ML_INCLUDES_START
 ML_INCLUDES_END
 #endif //WITH_MLSDK
 
-FThreadSafeCounter64 FCameraRunnable::PreviewHandle = ML_INVALID_HANDLE;
+FThreadSafeCounter64 FCameraRunnable::PreviewHandle = 0;
 
 FCameraRunnable::FCameraRunnable()
 #if WITH_MLSDK
@@ -32,7 +32,9 @@ FCameraRunnable::FCameraRunnable()
 
 void FCameraRunnable::Exit()
 {
+#if WITH_MLSDK
 	TryDisconnect();
+#endif
 }
 
 void FCameraRunnable::PushNewCaptureTask(FCameraTask::EType InTaskType)
@@ -83,17 +85,11 @@ void FCameraRunnable::Resume()
 	}
 }
 
-#if WITH_MLSDK
-void FCameraRunnable::OnPreviewBufferAvailable(MLHandle Output, void *Data)
-{
-	(void)Data;
-	PreviewHandle.Set(static_cast<int64>(Output));
-}
-
 bool FCameraRunnable::ProcessCurrentTask()
 {
 	bool bSuccess = false;
 
+#if WITH_MLSDK
 	switch (CurrentTask.CaptureType)
 	{
 	case FCameraTask::EType::None: bSuccess = false; checkf(false, TEXT("Invalid camera task encountered!")); break;
@@ -104,10 +100,18 @@ bool FCameraRunnable::ProcessCurrentTask()
 	case FCameraTask::EType::StartVideoToFile: bSuccess = StartRecordingVideo(); break;
 	case FCameraTask::EType::StopVideoToFile: bSuccess = StopRecordingVideo(); break;
 	}
+#endif
 
 	return bSuccess;
 }
 
+#if WITH_MLSDK
+void FCameraRunnable::OnPreviewBufferAvailable(MLHandle Output, void *Data)
+{
+	(void)Data;
+	PreviewHandle.Set(static_cast<int64>(Output));
+}
+
 bool FCameraRunnable::TryConnect()
 {
 	//IncomingTasks.Peek(CurrentTask); // this is purely so that any log messages go to the correct requester
