Copying //UE4/Partner-MagicLeap-Staging to Dev-VR (//UE4/Dev-VR)

diff --git a/Engine/Plugins/Lumin/MagicLeap/Source/MagicLeapPrivileges/Public/IMagicLeapPrivilegesModule.h b/Engine/Plugins/Lumin/MagicLeap/Source/MagicLeapPrivileges/Public/IMagicLeapPrivilegesModule.h
new file mode 100644
--- /dev/null
+++ b/Engine/Plugins/Lumin/MagicLeap/Source/MagicLeapPrivileges/Public/IMagicLeapPrivilegesModule.h
@@ -0,0 +1,51 @@
+// Copyright 1998-2019 Epic Games, Inc. All Rights Reserved.
+
+#pragma once
+
+#include "Modules/ModuleManager.h"
+#include "MagicLeapPrivilegeTypes.h"
+
+/**
+ * The public interface to this module.  In most cases, this interface is only public to sibling modules
+ * within this plugin.
+ */
+class IMagicLeapPrivilegesModule : public IModuleInterface
+{
+public:
+	/**
+	 * Singleton-like access to this module's interface.  This is just for convenience!
+	 * Beware of calling this during the shutdown phase, though.  Your module might have been unloaded already.
+	 *
+	 * @return Returns singleton instance, loading the module on demand if needed
+	 */
+	static inline IMagicLeapPrivilegesModule& Get()
+	{
+		return FModuleManager::LoadModuleChecked<IMagicLeapPrivilegesModule>("MagicLeapPrivileges");
+	}
+
+	/**
+	 * Checks to see if this module is loaded and ready.  It is only valid to call Get() if IsAvailable() returns true.
+	 *
+	 * @return True if the module is loaded and ready to use
+	 */
+	static inline bool IsAvailable()
+	{
+		return FModuleManager::Get().IsModuleLoaded("MagicLeapPrivileges");
+	}
+
+	/**
+	 * Requests the given privilege asynchronously.
+	 *
+	 * @param ResultDelegate The static delegate to be notified upon completion of the privilege request.
+	 * @return True if the request was initiated successfully, false otherwise.
+	 */
+	virtual bool RequestPrivilegeAsync(EMagicLeapPrivilege Privilege, const FMagicLeapPrivilegeRequestStaticDelegate& ResultDelegate) = 0;
+	
+	/**
+	 * Requests the given privilege asynchronously.
+	 *
+	 * @param ResultDelegate The dynamic delegate to be notified upon completion of the privilege request.
+	 * @return True if the request was initiated successfully, false otherwise.
+	 */
+	virtual bool RequestPrivilegeAsync(EMagicLeapPrivilege Privilege, const FMagicLeapPrivilegeRequestDelegate& ResultDelegate) = 0;
+};
diff --git a/Engine/Plugins/Lumin/MagicLeap/Source/MagicLeapPrivileges/Public/MagicLeapPrivilegeTypes.h b/Engine/Plugins/Lumin/MagicLeap/Source/MagicLeapPrivileges/Public/MagicLeapPrivilegeTypes.h
new file mode 100644
--- /dev/null
+++ b/Engine/Plugins/Lumin/MagicLeap/Source/MagicLeapPrivileges/Public/MagicLeapPrivilegeTypes.h
@@ -0,0 +1,54 @@
+// Copyright 1998-2019 Epic Games, Inc. All Rights Reserved.
+
+#pragma once
+
+#include "Engine/Engine.h"
+#include "MagicLeapPrivilegeTypes.generated.h"
+
+/** Privileges an app can request for from the system. */
+UENUM(BlueprintType)
+enum class EMagicLeapPrivilege : uint8
+{
+	Invalid,
+	BatteryInfo,
+	CameraCapture,
+	WorldReconstruction,
+	InAppPurchase,
+	AudioCaptureMic,
+	DrmCertificates,
+	Occlusion,
+	LowLatencyLightwear,
+	Internet,
+	IdentityRead,
+	BackgroundDownload,
+	BackgroundUpload,
+	MediaDrm,
+	Media,
+	MediaMetadata,
+	PowerInfo,
+	LocalAreaNetwork,
+	VoiceInput,
+	Documents,
+	ConnectBackgroundMusicService,
+	RegisterBackgroundMusicService,
+	PwFoundObjRead,
+	NormalNotificationsUsage,
+	MusicService,
+	ControllerPose,
+	ScreensProvider,
+	GesturesSubscribe,
+	GesturesConfig,
+	AddressBookRead,
+	AddressBookWrite,
+	CoarseLocation,
+	HandMesh,
+	WifiStatusRead,
+};
+
+/**
+  Delegate for the result of requesting a privilege asynchronously.
+  @param RequestedPrivilege The privilege that was requested.
+  @param WasGranted True if the privilege was granted, false otherwise.
+*/
+DECLARE_DELEGATE_TwoParams(FMagicLeapPrivilegeRequestStaticDelegate, EMagicLeapPrivilege, bool);
+DECLARE_DYNAMIC_DELEGATE_TwoParams(FMagicLeapPrivilegeRequestDelegate, EMagicLeapPrivilege, RequestedPrivilege, bool, WasGranted);
diff --git a/Engine/Plugins/Lumin/MagicLeap/Source/MagicLeapPrivileges/Public/MagicLeapPrivilegeUtils.h b/Engine/Plugins/Lumin/MagicLeap/Source/MagicLeapPrivileges/Public/MagicLeapPrivilegeUtils.h
new file mode 100644
--- /dev/null
+++ b/Engine/Plugins/Lumin/MagicLeap/Source/MagicLeapPrivileges/Public/MagicLeapPrivilegeUtils.h
@@ -0,0 +1,17 @@
+// Copyright 1998-2019 Epic Games, Inc. All Rights Reserved.
+
+#pragma once
+
+#include "MagicLeapPrivilegeTypes.h"
+#include "Lumin/CAPIShims/LuminAPIPrivileges.h"
+
+namespace MagicLeap
+{
+#if WITH_MLSDK
+	MLPrivilegeID MAGICLEAPPRIVILEGES_API UnrealToMLPrivilege(EMagicLeapPrivilege Privilege);
+
+	FString MAGICLEAPPRIVILEGES_API MLPrivilegeToString(MLPrivilegeID PrivilegeID);
+
+	FString MAGICLEAPPRIVILEGES_API MLPrivilegeToString(EMagicLeapPrivilege PrivilegeID);
+#endif //WITH_MLSDK
+}
diff --git a/Engine/Plugins/Lumin/MagicLeap/Source/MagicLeapPrivileges/Public/MagicLeapPrivilegesFunctionLibrary.h b/Engine/Plugins/Lumin/MagicLeap/Source/MagicLeapPrivileges/Public/MagicLeapPrivilegesFunctionLibrary.h
new file mode 100644
--- /dev/null
+++ b/Engine/Plugins/Lumin/MagicLeap/Source/MagicLeapPrivileges/Public/MagicLeapPrivilegesFunctionLibrary.h
@@ -0,0 +1,44 @@
+// Copyright 1998-2019 Epic Games, Inc. All Rights Reserved.
+
+#pragma once
+
+#include "Engine/Engine.h"
+#include "Kismet/BlueprintFunctionLibrary.h"
+#include "MagicLeapPrivilegeTypes.h"
+#include "MagicLeapPrivilegesFunctionLibrary.generated.h"
+
+UCLASS(ClassGroup = MagicLeap)
+class MAGICLEAPPRIVILEGES_API UMagicLeapPrivilegesFunctionLibrary : public UBlueprintFunctionLibrary
+{
+	GENERATED_BODY()
+
+public:
+	/**
+	  Check whether the application has the specified privilege.
+	  This does not solicit consent from the end-user and is non-blocking.
+	  @param Privilege The privilege to check.
+	  @return True if the privilege is granted, false otherwise.
+	*/
+	UFUNCTION(BlueprintCallable, Category = "Privileges|MagicLeap")
+	static bool CheckPrivilege(EMagicLeapPrivilege Privilege);
+
+	/**
+	  Request the specified privilege.
+	  This may possibly solicit consent from the end-user; if so it will block.
+	  @param Privilege The privilege to request.
+	  @return True if the privilege is granted, false otherwise.
+	 */
+	UFUNCTION(BlueprintCallable, Category = "Privileges|MagicLeap")
+	static bool RequestPrivilege(EMagicLeapPrivilege Privilege);
+
+	/**
+	  Request the specified privilege asynchronously.
+	  This may possibly solicit consent from the end-user. Result will be delivered
+	  to the specified delegate.
+	  @param Privilege The privilege to request.
+	  @param ResultDelegate Callback which reports the result of the request.
+	  @return True if the privilege request was successfully dispatched, false otherwise.
+	 */
+	UFUNCTION(BlueprintCallable, Category = "Privileges|MagicLeap")
+	static bool RequestPrivilegeAsync(EMagicLeapPrivilege Privilege, const FMagicLeapPrivilegeRequestDelegate& ResultDelegate);
+};
diff --git a/Engine/Plugins/Lumin/MagicLeap/Source/MagicLeapSecureStorage/Private/MagicLeapSecureStorage.cpp b/Engine/Plugins/Lumin/MagicLeap/Source/MagicLeapSecureStorage/Private/MagicLeapSecureStorage.cpp
--- a/Engine/Plugins/Lumin/MagicLeap/Source/MagicLeapSecureStorage/Private/MagicLeapSecureStorage.cpp
+++ b/Engine/Plugins/Lumin/MagicLeap/Source/MagicLeapSecureStorage/Private/MagicLeapSecureStorage.cpp
@@ -7,34 +7,11 @@
 #include "Modules/ModuleManager.h"
 #include "HAL/PlatformProcess.h"
 #include "Misc/Paths.h"
-#include "MagicLeapPluginUtil.h"
-
-#if WITH_MLSDK
-#include "ml_secure_storage.h"
-#endif //WITH_MLSDK
-
-DEFINE_LOG_CATEGORY_STATIC(LogSecureStorage, Display, All);
+#include "Kismet/GameplayStatics.h"
+#include "Lumin/CAPIShims/LuminAPISecureStorage.h"
 
 class FMagicLeapSecureStoragePlugin : public IMagicLeapSecureStoragePlugin
 {
-public:
-	virtual void StartupModule() override
-	{
-		IModuleInterface::StartupModule();
-		APISetup.Startup();
-#if WITH_MLSDK
-		APISetup.LoadDLL(TEXT("ml_secure_storage"));
-#endif
-	}
-
-	virtual void ShutdownModule() override
-	{
-		APISetup.Shutdown();
-		IModuleInterface::ShutdownModule();
-	}
-
-private:
-	FMagicLeapAPISetup APISetup;
 };
 
 IMPLEMENT_MODULE(FMagicLeapSecureStoragePlugin, MagicLeapSecureStorage);
@@ -59,9 +36,8 @@ bool UMagicLeapSecureStorage::GetSecureBlob<FString>(const FString& Key, FString
 
 	bool bResult = false;
 #if WITH_MLSDK
-	bResult = MLSecureStorageGetBlob(TCHAR_TO_ANSI(*Key), &blob, &blobLength) == MLResult_Ok;
-#endif //WITH_MLSDK
-	if (bResult)
+	MLResult Result = MLSecureStorageGetBlob(TCHAR_TO_ANSI(*Key), &blob, &blobLength);
+	if (Result == MLResult_Ok)
 	{
 		if (blob == nullptr)
 		{
@@ -70,17 +46,18 @@ bool UMagicLeapSecureStorage::GetSecureBlob<FString>(const FString& Key, FString
 		}
 		else
 		{
+			bResult = true;
 			DataToRetrieve = FString(ANSI_TO_TCHAR(reinterpret_cast<ANSICHAR*>(blob)));
 			// Replace with library function call when it comes online.
-#if WITH_MLSDK
 			MLSecureStorageFreeBlobBuffer(blob);
-#endif //WITH_MLSDK
 		}
 	}
 	else
 	{
-		UE_LOG(LogSecureStorage, Error, TEXT("Error retrieving secure blob with key %s"), *Key);
+		bResult = false;
+		UE_LOG(LogSecureStorage, Error, TEXT("MLSecureStorageGetBlob for key %s failed with error %s"), *Key, UTF8_TO_TCHAR(MLSecureStorageGetResultString(Result)));
 	}
+#endif //WITH_MLSDK
 
 	return bResult;
 }
@@ -100,6 +77,11 @@ bool UMagicLeapSecureStorage::PutSecureInt(const FString& Key, int32 DataToStore
 	return PutSecureBlob<int32>(Key, &DataToStore);
 }
 
+bool UMagicLeapSecureStorage::PutSecureInt64(const FString& Key, int64 DataToStore)
+{
+	return PutSecureBlob<int64>(Key, &DataToStore);
+}
+
 bool UMagicLeapSecureStorage::PutSecureFloat(const FString& Key, float DataToStore)
 {
 	return PutSecureBlob<float>(Key, &DataToStore);
@@ -140,6 +122,11 @@ bool UMagicLeapSecureStorage::GetSecureInt(const FString& Key, int32& DataToRetr
 	return GetSecureBlob<int32>(Key, DataToRetrieve);
 }
 
+bool UMagicLeapSecureStorage::GetSecureInt64(const FString& Key, int64& DataToRetrieve)
+{
+	return GetSecureBlob<int64>(Key, DataToRetrieve);
+}
+
 bool UMagicLeapSecureStorage::GetSecureFloat(const FString& Key, float& DataToRetrieve)
 {
 	return GetSecureBlob<float>(Key, DataToRetrieve);
@@ -165,6 +152,103 @@ bool UMagicLeapSecureStorage::GetSecureTransform(const FString& Key, FTransform&
 	return GetSecureBlob<FTransform>(Key, DataToRetrieve);
 }
 
+bool UMagicLeapSecureStorage::PutSecureSaveGame(const FString& Key, USaveGame* ObjectToStore)
+{
+	TArray<uint8> Bytes;
+	bool bSuccess = UGameplayStatics::SaveGameToMemory(ObjectToStore, Bytes);
+
+	if (!bSuccess || Bytes.Num() == 0)
+	{
+		UE_LOG(LogSecureStorage, Error, TEXT("Serialization of `%s` was unsuccessful."), *Key);
+		return false;
+	}
+
+	return PutSecureBlobImpl(Key, Bytes.GetData(), Bytes.Num());
+}
+
+bool UMagicLeapSecureStorage::GenericPutSecureArray(const FString& Key, const UArrayProperty* ArrayProperty, void* TargetArray)
+{
+
+	FScriptArrayHelper ArrayHelper(ArrayProperty, TargetArray);
+	const auto& ElementSize = ArrayProperty->Inner->ElementSize;
+	size_t ByteNum = ElementSize * ArrayHelper.Num();
+
+	// Contiguous memory to send to SecureStorage
+	TArray<uint8> Bytes;
+	Bytes.SetNum(ByteNum);
+
+	// Generate the contiguous array from the UArrayProperty
+	for (int32 i = 0; i < ArrayHelper.Num(); ++i)
+	{
+		ArrayProperty->Inner->CopySingleValueFromScriptVM(Bytes.GetData() + i * ElementSize, ArrayHelper.GetRawPtr(i));
+	}
+
+	return PutSecureBlobImpl(Key, Bytes.GetData(), ByteNum);
+
+}
+
+bool UMagicLeapSecureStorage::GetSecureSaveGame(const FString& Key, USaveGame*& ObjectToRetrieve)
+{
+	uint8* Data;
+	size_t ArrayNum;
+	bool bSuccess = GetSecureBlobImpl(Key, Data, ArrayNum);
+
+	TArray<uint8> Bytes(Data, ArrayNum);
+	ObjectToRetrieve = UGameplayStatics::LoadGameFromMemory(Bytes);
+
+	if (!ObjectToRetrieve)
+	{
+		UE_LOG(LogSecureStorage, Error, TEXT("Deserialization of `%s` was unsuccessful."), *Key);
+
+		if (bSuccess) 
+		{
+			FreeBlobBufferImpl(Data);
+		}
+
+		return false;
+	}
+
+	return bSuccess;
+}
+
+bool UMagicLeapSecureStorage::GenericGetSecureArray(const FString& Key, UArrayProperty* ArrayProperty, void* TargetArray)
+{
+
+	FScriptArrayHelper ArrayHelper(ArrayProperty, TargetArray);
+
+	uint8* Data;
+	size_t ArrayNum;
+	bool bSuccess = GetSecureBlobImpl(Key, Data, ArrayNum);
+
+	//Early out to avoid processing
+	if (!bSuccess)
+	{
+		return false;
+	}
+
+	const int32& ElementSize = ArrayProperty->Inner->ElementSize;
+
+	// Without type metadata, validate that the returned byte count is divisible by the requested element type.
+	if (ArrayNum % ElementSize)
+	{
+		UE_LOG(LogSecureStorage, Error, TEXT("Size of blob data %s does not match the size of requested data type."), *Key);
+		FreeBlobBufferImpl(Data);
+		return false;
+	}
+
+	// Convert the byte size to the requested type size
+	ArrayHelper.Resize(ArrayNum / ElementSize);
+
+	// Copy the contiguous array into the property value locations
+	for (int32 i = 0; i < ArrayHelper.Num(); ++i)
+	{
+		ArrayProperty->Inner->CopySingleValueToScriptVM(ArrayHelper.GetRawPtr(i), Data + i * ElementSize);
+	}
+
+	return true;
+
+}
+
 bool UMagicLeapSecureStorage::DeleteSecureData(const FString& Key)
 {
 #if WITH_MLSDK
@@ -174,44 +258,51 @@ bool UMagicLeapSecureStorage::DeleteSecureData(const FString& Key)
 #endif //WITH_MLSDK
 }
 
-bool UMagicLeapSecureStorage::PutSecureBlobImpl(const FString& Key, const void* DataToStore, size_t DataTypeSize)
+bool UMagicLeapSecureStorage::PutSecureBlobImpl(const FString& Key, const uint8* DataToStore, size_t DataTypeSize)
 {
 #if WITH_MLSDK
-	return MLSecureStoragePutBlob(TCHAR_TO_ANSI(*Key), reinterpret_cast<const uint8*>(DataToStore), DataTypeSize) == MLResult_Ok;
+
+	MLResult Result = MLSecureStoragePutBlob(TCHAR_TO_ANSI(*Key), DataToStore, DataTypeSize);
+
+	if (Result != MLResult_Ok)
+	{
+		UE_LOG(LogSecureStorage, Error, TEXT("MLSecureStoragePutBlob for key %s failed with error %s"), *Key, UTF8_TO_TCHAR(MLSecureStorageGetResultString(Result)));
+	}
+
+	return Result == MLResult_Ok;
 #else
 	return false;
 #endif //WITH_MLSDK
 }
 
-uint8* UMagicLeapSecureStorage::GetSecureBlobImpl(const FString& Key, size_t DataTypeSize)
+bool UMagicLeapSecureStorage::GetSecureBlobImpl(const FString& Key, uint8*& DataToRetrieve, size_t& DataTypeSize) 
 {
-	uint8* blob = nullptr;
-	size_t blobLength = 0;
+
+	DataToRetrieve = nullptr;
+	DataTypeSize = 0;
 
 #if WITH_MLSDK
-	MLResult result = MLSecureStorageGetBlob(TCHAR_TO_ANSI(*Key), &blob, &blobLength);
-	if (MLResult_Ok == result)
+	MLResult Result = MLSecureStorageGetBlob(TCHAR_TO_ANSI(*Key), &DataToRetrieve, &DataTypeSize);
+	if (Result == MLResult_Ok)
 	{
-		if (blob == nullptr)
+		// Additional validation
+		if (!DataToRetrieve)
 		{
 			UE_LOG(LogSecureStorage, Error, TEXT("Error retrieving secure blob with key %s. Blob was null."), *Key);
-			result = MLSecureStorageResult_IOFailure;
-		}
-		else if (blobLength != DataTypeSize)
-		{
-			UE_LOG(LogSecureStorage, Error, TEXT("Size of blob data %s does not match the size of requested data type. Requested size = %d vs Actual size = %d"), *Key, DataTypeSize, blobLength);
-			result = MLResult_UnspecifiedFailure;
-			MLSecureStorageFreeBlobBuffer(blob);
+			return false;
 		}
+
+		return true;
 	}
-	else if (result != MLSecureStorageResult_BlobNotFound)
+	
+	if (Result != MLSecureStorageResult_BlobNotFound)
 	{
-		UE_LOG(LogSecureStorage, Error, TEXT("Error retrieving secure blob with key %s. Error Code = %d"), *Key, static_cast<int32>(result));
+		UE_LOG(LogSecureStorage, Error, TEXT("MLSecureStorageGetBlob for key %s failed with error %s"), *Key, UTF8_TO_TCHAR(MLSecureStorageGetResultString(Result)));
 	}
-	return (MLResult_Ok == result) ? blob : nullptr;
-#else
-	return nullptr;
+
 #endif //WITH_MLSDK
+
+	return false;
 }
 
 void UMagicLeapSecureStorage::FreeBlobBufferImpl(uint8* Buffer)
diff --git a/Engine/Plugins/Lumin/MagicLeap/Source/MagicLeapSecureStorage/Public/MagicLeapSecureStorage.h b/Engine/Plugins/Lumin/MagicLeap/Source/MagicLeapSecureStorage/Public/MagicLeapSecureStorage.h
--- a/Engine/Plugins/Lumin/MagicLeap/Source/MagicLeapSecureStorage/Public/MagicLeapSecureStorage.h
+++ b/Engine/Plugins/Lumin/MagicLeap/Source/MagicLeapSecureStorage/Public/MagicLeapSecureStorage.h
@@ -7,11 +7,12 @@
 #include "Kismet/BlueprintFunctionLibrary.h"
 #include "MagicLeapSecureStorage.generated.h"
 
+DEFINE_LOG_CATEGORY_STATIC(LogSecureStorage, Display, All);
+
 /**
   Function library for the Magic Leap Secure Storage API.
   Currently supports bool, uint8, int32, float, FString, FVector, FRotator and FTransform via Blueprints.
   Provides a template function for any non specialized types to be used via C++.
-  TODO: Support TArray and a generic USTRUCT.
 */
 UCLASS(ClassGroup = MagicLeap)
 class MAGICLEAPSECURESTORAGE_API UMagicLeapSecureStorage : public UBlueprintFunctionLibrary
@@ -46,6 +47,15 @@ public:
 	UFUNCTION(BlueprintCallable, Category = "SecureStorage|MagicLeap")
 	static bool PutSecureInt(const FString& Key, int32 DataToStore);
 
+	/**
+	  Stores the 64 bit integer under the specified key. An existing key would be overwritten.
+	  @param   Key String key associated with the data.
+	  @param   DataToStore The data to store.
+	  @return  True if the data was stored successfully, false otherwise.
+	*/
+	UFUNCTION(BlueprintCallable, Category = "SecureStorage|MagicLeap")
+	static bool PutSecureInt64(const FString& Key, int64 DataToStore);
+
 	/**
 	  Stores the float under the specified key. An existing key would be overwritten.
 	  @param   Key String key associated with the data.
@@ -70,7 +80,7 @@ public:
 	  @param   DataToStore The data to store.
 	  @return  True if the data was stored successfully, false otherwise.
 	*/
-	UFUNCTION(BlueprintCallable, Category = "SecureStorage|MagicLeap")
+	UFUNCTION(BlueprintCallable, meta = (DeprecatedFunction, DeprecationMessage = "This function has been replaced by `PutSecureArray`"), Category = "SecureStorage|MagicLeap")
 	static bool PutSecureVector(const FString& Key, const FVector& DataToStore);
 
 	/**
@@ -79,7 +89,7 @@ public:
 	  @param   DataToStore The data to store.
 	  @return  True if the data was stored successfully, false otherwise.
 	*/
-	UFUNCTION(BlueprintCallable, Category = "SecureStorage|MagicLeap")
+	UFUNCTION(BlueprintCallable, meta = (DeprecatedFunction, DeprecationMessage = "This function has been replaced by `PutSecureArray`"), Category = "SecureStorage|MagicLeap")
 	static bool PutSecureRotator(const FString& Key, const FRotator& DataToStore);
 
 	/**
@@ -88,7 +98,7 @@ public:
 	  @param   DataToStore The data to store.
 	  @return  True if the data was stored successfully, false otherwise.
 	*/
-	UFUNCTION(BlueprintCallable, Category = "SecureStorage|MagicLeap")
+	UFUNCTION(BlueprintCallable, meta = (DeprecatedFunction, DeprecationMessage = "This function has been replaced by `PutSecureArray`"), Category = "SecureStorage|MagicLeap")
 	static bool PutSecureTransform(const FString& Key, const FTransform& DataToStore);
 
 	/**
@@ -118,6 +128,15 @@ public:
 	UFUNCTION(BlueprintCallable, Category = "SecureStorage|MagicLeap")
 	static bool GetSecureInt(const FString& Key, int32& DataToRetrieve);
 
+	/**
+	  Retrieves the 64 bit integer associated with the specified key.
+	  @param   Key The string key to look for.
+	  @param   DataToRetrieve Reference to the variable that will be populated with the stored data.
+	  @return  True if the key was found and output parameter was successfully populated with the data, false otherwise.
+	*/
+	UFUNCTION(BlueprintCallable, Category = "SecureStorage|MagicLeap")
+	static bool GetSecureInt64(const FString& Key, int64& DataToRetrieve);
+
 	/**
 	  Retrieves the float associated with the specified key.
 	  @param   Key The string key to look for.
@@ -142,7 +161,7 @@ public:
 	  @param   DataToRetrieve Reference to the variable that will be populated with the stored data.
 	  @return  True if the key was found and output parameter was successfully populated with the data, false otherwise.
 	*/
-	UFUNCTION(BlueprintCallable, Category = "SecureStorage|MagicLeap")
+	UFUNCTION(BlueprintCallable, meta = (DeprecatedFunction, DeprecationMessage = "This function has been replaced by `GetSecureArray`"), Category = "SecureStorage|MagicLeap")
 	static bool GetSecureVector(const FString& Key, FVector& DataToRetrieve);
 
 	/**
@@ -151,7 +170,7 @@ public:
 	  @param   DataToRetrieve Reference to the variable that will be populated with the stored data.
 	  @return  True if the key was found and output parameter was successfully populated with the data, false otherwise.
 	*/
-	UFUNCTION(BlueprintCallable, Category = "SecureStorage|MagicLeap")
+	UFUNCTION(BlueprintCallable, meta = (DeprecatedFunction, DeprecationMessage = "This function has been replaced by `GetSecureArray`"), Category = "SecureStorage|MagicLeap")
 	static bool GetSecureRotator(const FString& Key, FRotator& DataToRetrieve);
 
 	/**
@@ -160,16 +179,94 @@ public:
 	  @param   DataToRetrieve Reference to the variable that will be populated with the stored data.
 	  @return  True if the key was found and output parameter was successfully populated with the data, false otherwise.
 	*/
-	UFUNCTION(BlueprintCallable, Category = "SecureStorage|MagicLeap")
+	UFUNCTION(BlueprintCallable, meta=(DeprecatedFunction, DeprecationMessage="This function has been replaced by `GetSecureArray`"),Category = "SecureStorage|MagicLeap")
 	static bool GetSecureTransform(const FString& Key, FTransform& DataToRetrieve);
 
 	/**
-	  Deletes the data associated with the specified key.
-	  @param   Key The string key of the data to delete.
-	  @return  True if the data was deleted successfully or did not exist altogether, false otherwise.
+	  Stores the USaveGame object under the specified key. An existing key would be overwritten.
+	  @param   Key String key associated with the data.
+	  @param   ObjectToStore The USaveGame object to serialize and store.
+	  @return  True if the data was stored successfully, false otherwise.
 	*/
 	UFUNCTION(BlueprintCallable, Category = "SecureStorage|MagicLeap")
-	static bool DeleteSecureData(const FString& Key);
+	static bool PutSecureSaveGame(const FString& Key, USaveGame* ObjectToStore);
+
+	/**
+	  Stores the data under the specified key. An existing key would be overwritten.
+	  @param   Key String key associated with the data.
+	  @param   DataToStore The data to store.
+	  @return  True if the data was stored successfully, false otherwise.
+	*/
+	UFUNCTION(BlueprintCallable, CustomThunk, meta = (ArrayParm = "DataToStore"), Category = "SecureStorage|MagicLeap")
+	static bool PutSecureArray(const FString& Key, const TArray<int32>& DataToStore);
+
+	DECLARE_FUNCTION(execPutSecureArray)
+	{
+
+		P_GET_PROPERTY(UStrProperty, Key);
+
+		Stack.StepCompiledIn<UArrayProperty>(nullptr);
+		void* ArrayAddress = Stack.MostRecentPropertyAddress;
+		UArrayProperty* ArrayProperty = Cast<UArrayProperty>(Stack.MostRecentProperty);
+
+		if (!ArrayProperty)
+		{
+			Stack.bArrayContextFailed = true;
+			return;
+		}
+
+		P_FINISH;
+
+		P_NATIVE_BEGIN;
+		*static_cast<bool*>(RESULT_PARAM) = GenericPutSecureArray(Key, ArrayProperty, ArrayAddress);
+		P_NATIVE_END;
+
+	}
+
+	static bool GenericPutSecureArray(const FString& Key, const UArrayProperty* ArrayProperty, void* TargetArray);
+
+	/**
+	  Retrieves a USaveGame object associated with the specified key.
+	  @param   Key The string key to look for.
+	  @param   ObjectToRetrieve Reference to a USaveGame object that will be populated with the serialized data.
+	  @return  True if the key was found and output parameter was successfully populated with the data, false otherwise.
+	*/
+	UFUNCTION(BlueprintCallable, Category = "SecureStorage|MagicLeap")
+	static bool GetSecureSaveGame(const FString& Key, USaveGame*& ObjectToRetrieve);
+
+	/**
+	  Retrieves an array associated with the specified key.
+	  @param   Key The string key to look for.
+	  @param   DataToRetrieve Reference to an array that will be populated with the stored data.
+	  @return  True if the key was found and output parameter was successfully populated with the data, false otherwise.
+	*/
+	UFUNCTION(BlueprintCallable, CustomThunk, meta = (ArrayParm = "DataToRetrieve"), Category = "SecureStorage|MagicLeap")
+	static bool GetSecureArray(const FString& Key, TArray<int32>& DataToRetrieve);
+
+	DECLARE_FUNCTION(execGetSecureArray)
+	{
+
+		P_GET_PROPERTY(UStrProperty, Key);
+
+		Stack.StepCompiledIn<UArrayProperty>(nullptr);
+		void* ArrayAddress = Stack.MostRecentPropertyAddress;
+		UArrayProperty* ArrayProperty = Cast<UArrayProperty>(Stack.MostRecentProperty);
+
+		if (!ArrayProperty)
+		{
+			Stack.bArrayContextFailed = true;
+			return;
+		}
+
+		P_FINISH;
+
+		P_NATIVE_BEGIN;
+		*static_cast<bool*>(RESULT_PARAM) = GenericGetSecureArray(Key, ArrayProperty, ArrayAddress);
+		P_NATIVE_END;
+
+	}
+
+	static bool GenericGetSecureArray(const FString& Key, UArrayProperty* ArrayProperty, void* TargetArray);
 
 	/**
 	  Template function to store the data under the specified key. An existing key would be overwritten.
@@ -189,10 +286,18 @@ public:
 	template<class T>
 	static bool GetSecureBlob(const FString& Key, T& DataToRetrieve);
 
+	/**
+	  Deletes the data associated with the specified key.
+	  @param   Key The string key of the data to delete.
+	  @return  True if the data was deleted successfully or did not exist altogether, false otherwise.
+	*/
+	UFUNCTION(BlueprintCallable, Category = "SecureStorage|MagicLeap")
+	static bool DeleteSecureData(const FString& Key);
+
 private:
-	static bool PutSecureBlobImpl(const FString& Key, const void* DataToStore, size_t DataTypeSize);
+	static bool PutSecureBlobImpl(const FString& Key, const uint8* DataToStore, size_t DataTypeSize);
 
-	static uint8* GetSecureBlobImpl(const FString& Key, size_t DataTypeSize);
+	static bool GetSecureBlobImpl(const FString& Key, uint8*& DataToRetrieve, size_t& DataTypeSize);
 
 	static void FreeBlobBufferImpl(uint8* Buffer);
 };
@@ -200,17 +305,29 @@ private:
 template<class T>
 bool UMagicLeapSecureStorage::PutSecureBlob(const FString& Key, const T* DataToStore)
 {
-	return UMagicLeapSecureStorage::PutSecureBlobImpl(Key, reinterpret_cast<const void*>(DataToStore), sizeof(T));
+	return UMagicLeapSecureStorage::PutSecureBlobImpl(Key, reinterpret_cast<const uint8*>(DataToStore), sizeof(T));
 }
 
 template<class T>
 bool UMagicLeapSecureStorage::GetSecureBlob(const FString& Key, T& DataToRetrieve)
 {
-	uint8* blob = UMagicLeapSecureStorage::GetSecureBlobImpl(Key, sizeof(T));
-	if (blob != nullptr)
+
+	uint8* Data;
+	size_t Size;
+	bool bSuccess = UMagicLeapSecureStorage::GetSecureBlobImpl(Key, Data, Size);
+
+	if (!bSuccess) 
 	{
-		DataToRetrieve = *reinterpret_cast<T*>(blob);
-		UMagicLeapSecureStorage::FreeBlobBufferImpl(blob);
+		return false;
 	}
-	return blob != nullptr;
+
+	if (Size != sizeof(T))
+	{
+		UE_LOG(LogSecureStorage, Error, TEXT("Size of blob data %s does not match the size of requested data type. Requested size = %d vs Actual size = %d"), *Key, sizeof(T), Size);
+		FreeBlobBufferImpl(Data);
+		return false;
+	}
+
+	DataToRetrieve = *reinterpret_cast<T*>(Data);
+	return bSuccess;
 }
diff --git a/Engine/Plugins/Lumin/MagicLeapCamera/Config/BaseMagicLeapCamera.ini b/Engine/Plugins/Lumin/MagicLeapCamera/Config/BaseMagicLeapCamera.ini
new file mode 100644
--- /dev/null
+++ b/Engine/Plugins/Lumin/MagicLeapCamera/Config/BaseMagicLeapCamera.ini
@@ -0,0 +1,18 @@
+[CoreRedirects]
++ClassRedirects=(OldName="/Script/MagicLeapCamera.CameraConnect", NewName="/Script/MagicLeapCamera.MagicLeapCameraConnect")
++ClassRedirects=(OldName="/Script/MagicLeapCamera.CameraDisconnect", NewName="/Script/MagicLeapCamera.MagicLeapCameraDisconnect")
+
++ClassRedirects=(OldName="/Script/MagicLeapCamera.CameraCaptureImgToFile", NewName="/Script/MagicLeapCamera.MagicLeapCameraCaptureImgToFile")
++ClassRedirects=(OldName="/Script/MagicLeapCamera.CameraCaptureImgToFileMulti", NewName="/Script/MagicLeapCamera.MagicLeapCameraCaptureImgToFileMulti")
+
++ClassRedirects=(OldName="/Script/MagicLeapCamera.CameraCaptureImgToTexture", NewName="/Script/MagicLeapCamera.MagicLeapCameraCaptureImgToTexture")
++ClassRedirects=(OldName="/Script/MagicLeapCamera.CameraCaptureImgToTextureMulti", NewName="/Script/MagicLeapCamera.MagicLeapCameraCaptureImgToTextureMulti")
+
++ClassRedirects=(OldName="/Script/MagicLeapCamera.CameraStartRecording", NewName="/Script/MagicLeapCamera.MagicLeapCameraStartRecording")
++ClassRedirects=(OldName="/Script/MagicLeapCamera.CameraStartRecordingMulti", NewName="/Script/MagicLeapCamera.MagicLeapCameraStartRecordingMulti")
+
++ClassRedirects=(OldName="/Script/MagicLeapCamera.CameraStopRecording", NewName="/Script/MagicLeapCamera.MagicLeapCameraStopRecording")
++ClassRedirects=(OldName="/Script/MagicLeapCamera.CameraStopRecordingMulti", NewName="/Script/MagicLeapCamera.MagicLeapCameraStopRecordingMulti")
+
++ClassRedirects=(OldName="/Script/MagicLeapCamera.CameraLogMessage", NewName="/Script/MagicLeapCamera.MagicLeapCameraLogMessage")
++ClassRedirects=(OldName="/Script/MagicLeapCamera.CameraLogMessageMulti", NewName="/Script/MagicLeapCamera.MagicLeapCameraLogMessageMulti")
diff --git a/Engine/Plugins/Lumin/MagicLeapCamera/MagicLeapCamera.uplugin b/Engine/Plugins/Lumin/MagicLeapCamera/MagicLeapCamera.uplugin
--- a/Engine/Plugins/Lumin/MagicLeapCamera/MagicLeapCamera.uplugin
+++ b/Engine/Plugins/Lumin/MagicLeapCamera/MagicLeapCamera.uplugin
@@ -15,13 +15,33 @@
 	"IsBetaVersion": false,
 	"Installed": false,
 	"SupportedTargetPlatforms": [
-		"Lumin"
+		"Lumin",
+		"Win64",
+		"Mac"
 	],
 	"Modules": [
 		{
 			"Name": "MagicLeapCamera",
 			"Type": "Runtime",
-			"LoadingPhase": "PostConfigInit",
+			"LoadingPhase": "PostEngineInit",
+			"WhitelistPlatforms": [
+				"Win64",
+				"Mac",
+				"Lumin"
+			]
+		},
+		{
+			"Name": "MagicLeapCameraPreview",
+			"Type": "Runtime",
+			"LoadingPhase": "PostEngineInit",
+			"WhitelistPlatforms": [
+				"Lumin"
+			]
+		},
+		{
+			"Name": "MagicLeapCameraPreviewFactory",
+			"Type": "Runtime",
+			"LoadingPhase": "PostEngineInit",
 			"WhitelistPlatforms": [
 				"Win64",
 				"Mac",
@@ -30,6 +50,10 @@
 		}
 	],
 	"Plugins": [
+		{
+			"Name": "MLSDK",
+			"Enabled": true
+		},
 		{
 			"Name": "MagicLeap",
 			"Enabled": true
diff --git a/Engine/Plugins/Lumin/MagicLeapCamera/Source/MagicLeapCamera/MagicLeapCamera.Build.cs b/Engine/Plugins/Lumin/MagicLeapCamera/Source/MagicLeapCamera/MagicLeapCamera.Build.cs
new file mode 100644
--- /dev/null
+++ b/Engine/Plugins/Lumin/MagicLeapCamera/Source/MagicLeapCamera/MagicLeapCamera.Build.cs
@@ -0,0 +1,34 @@
+// Copyright 1998-2019 Epic Games, Inc. All Rights Reserved.
+
+using System.IO;
+
+namespace UnrealBuildTool.Rules
+{
+	public class MagicLeapCamera : ModuleRules
+	{
+		public MagicLeapCamera( ReadOnlyTargetRules Target ) : base(Target)
+		{
+			PublicIncludePaths.Add(Path.Combine(ModuleDirectory, "Public"));
+
+			PublicDependencyModuleNames.AddRange(
+				new string[]
+				{
+					"Core",
+					"CoreUObject",
+					"Engine"
+				}
+			);
+
+			PrivateDependencyModuleNames.AddRange(
+				new string[]
+				{
+					"Core",
+					"Engine",
+					"MLSDK",
+					"MagicLeap",
+					"HeadMountedDisplay",
+				}
+			);
+		}
+	}
+}
diff --git a/Engine/Plugins/Lumin/MagicLeapCamera/Source/MagicLeapCamera/Private/MagicLeapCameraComponent.cpp b/Engine/Plugins/Lumin/MagicLeapCamera/Source/MagicLeapCamera/Private/MagicLeapCameraComponent.cpp
new file mode 100644
--- /dev/null
+++ b/Engine/Plugins/Lumin/MagicLeapCamera/Source/MagicLeapCamera/Private/MagicLeapCameraComponent.cpp
@@ -0,0 +1,44 @@
+// Copyright 1998-2019 Epic Games, Inc. All Rights Reserved.
+
+#include "MagicLeapCameraComponent.h"
+#include "MagicLeapCameraPlugin.h"
+
+IMPLEMENT_MODULE(FMagicLeapCameraPlugin, MagicLeapCamera);
+
+void UMagicLeapCameraComponent::BeginPlay()
+{
+	Super::BeginPlay();
+	GetMagicLeapCameraPlugin().IncUserCount();
+	GetMagicLeapCameraPlugin().SetLogDelegate(OnLogMessage);
+}
+
+void UMagicLeapCameraComponent::EndPlay(const EEndPlayReason::Type EndPlayReason)
+{
+	GetMagicLeapCameraPlugin().DecUserCount();
+	Super::EndPlay(EndPlayReason);
+}
+
+bool UMagicLeapCameraComponent::CaptureImageToFileAsync()
+{
+	return GetMagicLeapCameraPlugin().CaptureImageToFileAsync(OnCaptureImgToFile);
+}
+
+bool UMagicLeapCameraComponent::CaptureImageToTextureAsync()
+{
+	return GetMagicLeapCameraPlugin().CaptureImageToTextureAsync(OnCaptureImgToTexture);
+}
+
+bool UMagicLeapCameraComponent::StartRecordingAsync()
+{
+	return GetMagicLeapCameraPlugin().StartRecordingAsync(OnStartRecording);
+}
+
+bool UMagicLeapCameraComponent::StopRecordingAsync()
+{
+	return GetMagicLeapCameraPlugin().StopRecordingAsync(OnStopRecording);
+}
+
+bool UMagicLeapCameraComponent::IsCapturing() const
+{
+	return GetMagicLeapCameraPlugin().IsCapturing();
+}
diff --git a/Engine/Plugins/Lumin/MagicLeapCamera/Source/MagicLeapCamera/Private/MagicLeapCameraFunctionLibrary.cpp b/Engine/Plugins/Lumin/MagicLeapCamera/Source/MagicLeapCamera/Private/MagicLeapCameraFunctionLibrary.cpp
new file mode 100644
--- /dev/null
+++ b/Engine/Plugins/Lumin/MagicLeapCamera/Source/MagicLeapCamera/Private/MagicLeapCameraFunctionLibrary.cpp
@@ -0,0 +1,54 @@
+// Copyright 1998-2019 Epic Games, Inc. All Rights Reserved.
+
+#include "MagicLeapCameraFunctionLibrary.h"
+#include "MagicLeapCameraPlugin.h"
+
+bool UMagicLeapCameraFunctionLibrary::CameraConnect(const FMagicLeapCameraConnect& ResultDelegate)
+{
+	return GetMagicLeapCameraPlugin().CameraConnect(ResultDelegate);
+}
+
+bool UMagicLeapCameraFunctionLibrary::CameraDisconnect(const FMagicLeapCameraDisconnect& ResultDelegate)
+{
+	return GetMagicLeapCameraPlugin().CameraDisconnect(ResultDelegate);
+}
+
+bool UMagicLeapCameraFunctionLibrary::CaptureImageToFileAsync(const FMagicLeapCameraCaptureImgToFile& InResultDelegate)
+{
+	FMagicLeapCameraCaptureImgToFileMulti ResultDelegate;
+	ResultDelegate.Add(InResultDelegate);
+	return GetMagicLeapCameraPlugin().CaptureImageToFileAsync(ResultDelegate);
+}
+
+bool UMagicLeapCameraFunctionLibrary::CaptureImageToTextureAsync(const FMagicLeapCameraCaptureImgToTexture& InResultDelegate)
+{
+	FMagicLeapCameraCaptureImgToTextureMulti ResultDelegate;
+	ResultDelegate.Add(InResultDelegate);
+	return GetMagicLeapCameraPlugin().CaptureImageToTextureAsync(ResultDelegate);
+}
+
+bool UMagicLeapCameraFunctionLibrary::StartRecordingAsync(const FMagicLeapCameraStartRecording& InResultDelegate)
+{
+	FMagicLeapCameraStartRecordingMulti ResultDelegate;
+	ResultDelegate.Add(InResultDelegate);
+	return GetMagicLeapCameraPlugin().StartRecordingAsync(ResultDelegate);
+}
+
+bool UMagicLeapCameraFunctionLibrary::StopRecordingAsync(const FMagicLeapCameraStopRecording& InResultDelegate)
+{
+	FMagicLeapCameraStopRecordingMulti ResultDelegate;
+	ResultDelegate.Add(InResultDelegate);
+	return GetMagicLeapCameraPlugin().StopRecordingAsync(ResultDelegate);
+}
+
+bool UMagicLeapCameraFunctionLibrary::SetLogDelegate(const FMagicLeapCameraLogMessage& InLogDelegate)
+{
+	FMagicLeapCameraLogMessageMulti LogDelegate;
+	LogDelegate.Add(InLogDelegate);
+	return GetMagicLeapCameraPlugin().SetLogDelegate(LogDelegate);
+}
+
+bool UMagicLeapCameraFunctionLibrary::IsCapturing()
+{
+	return GetMagicLeapCameraPlugin().IsCapturing();
+}
diff --git a/Engine/Plugins/Lumin/MagicLeapCamera/Source/MagicLeapCamera/Private/MagicLeapCameraPlugin.cpp b/Engine/Plugins/Lumin/MagicLeapCamera/Source/MagicLeapCamera/Private/MagicLeapCameraPlugin.cpp
new file mode 100644
--- /dev/null
+++ b/Engine/Plugins/Lumin/MagicLeapCamera/Source/MagicLeapCamera/Private/MagicLeapCameraPlugin.cpp
@@ -0,0 +1,296 @@
+// Copyright 1998-2019 Epic Games, Inc. All Rights Reserved.
+
+#include "MagicLeapCameraPlugin.h"
+#include "Async/Async.h"
+#include "Lumin/CAPIShims/LuminAPI.h"
+
+DEFINE_LOG_CATEGORY(LogMagicLeapCamera);
+
+FMagicLeapCameraPlugin::FMagicLeapCameraPlugin()
+: UserCount(0)
+, Runnable(nullptr)
+, CurrentTaskType(FCameraTask::EType::None)
+, PrevTaskType(FCameraTask::EType::None)
+{
+}
+
+void FMagicLeapCameraPlugin::StartupModule()
+{
+	IMagicLeapCameraPlugin::StartupModule();
+	Runnable = new FCameraRunnable();
+	TickDelegate = FTickerDelegate::CreateRaw(this, &FMagicLeapCameraPlugin::Tick);
+	TickDelegateHandle = FTicker::GetCoreTicker().AddTicker(TickDelegate);
+
+	AppEventHandler.SetOnAppPauseHandler([this]()
+	{
+		OnAppPause();
+	});
+}
+
+void FMagicLeapCameraPlugin::ShutdownModule()
+{
+	FCameraRunnable* InRunnable = Runnable;
+	Runnable = nullptr;
+	AsyncTask(ENamedThreads::AnyBackgroundThreadNormalTask, [InRunnable]()
+	{
+		delete InRunnable;
+	});
+	FTicker::GetCoreTicker().RemoveTicker(TickDelegateHandle);
+	IModuleInterface::ShutdownModule();
+}
+
+bool FMagicLeapCameraPlugin::Tick(float DeltaTime)
+{
+	FCameraTask CompletedTask;
+	if (Runnable->TryGetCompletedTask(CompletedTask))
+	{
+		switch (CompletedTask.CaptureType)
+		{
+		case FCameraTask::EType::Connect:
+		{
+			if (OnCameraConnect.IsBound())
+			{
+				OnCameraConnect.ExecuteIfBound(CompletedTask.bSuccess);
+				// Only clear CurrentTaskType if connect was manually called
+				// (otherwise it was auto-called from within the runnable and the actual
+				// current task is still in progress).
+				PrevTaskType = CurrentTaskType;
+				CurrentTaskType = FCameraTask::EType::None;
+			}
+		}
+		break;
+
+		case FCameraTask::EType::Disconnect:
+		{
+			OnCameraDisconnect.ExecuteIfBound(CompletedTask.bSuccess);
+			PrevTaskType = CurrentTaskType;
+			CurrentTaskType = FCameraTask::EType::None;
+		}
+		break;
+
+		case FCameraTask::EType::ImageToFile:
+		{
+			OnCaptureImgToFile.Broadcast(CompletedTask.bSuccess, CompletedTask.FilePath);
+			PrevTaskType = CurrentTaskType;
+			CurrentTaskType = FCameraTask::EType::None;
+		}
+		break;
+
+		case FCameraTask::EType::ImageToTexture:
+		{
+			OnCaptureImgToTexture.Broadcast(CompletedTask.bSuccess, CompletedTask.Texture);
+			PrevTaskType = CurrentTaskType;
+			CurrentTaskType = FCameraTask::EType::None;
+		}
+		break;
+
+		case FCameraTask::EType::StartVideoToFile:
+		{
+			OnStartRecording.Broadcast(CompletedTask.bSuccess);
+			PrevTaskType = CurrentTaskType;
+			// do not reset CurrentTaskType if start recording is successful
+			// as this constitutes an ongoing capture state
+			if (!CompletedTask.bSuccess)
+			{
+				CurrentTaskType = FCameraTask::EType::None;
+			}
+		}
+		break;
+
+		case FCameraTask::EType::StopVideoToFile:
+		{
+			OnStopRecording.Broadcast(CompletedTask.bSuccess, CompletedTask.FilePath);
+			PrevTaskType = CurrentTaskType;
+			CurrentTaskType = FCameraTask::EType::None;
+		}
+		break;
+
+		case FCameraTask::EType::Log:
+		{
+			UE_LOG(LogMagicLeapCamera, Log, TEXT("%s"), *CompletedTask.Log);
+			OnLogMessage.Broadcast(FString::Printf(TEXT("<br>%s"), *CompletedTask.Log));
+		}
+		break;
+		}
+	}
+
+	return true;
+}
+
+bool FMagicLeapCameraPlugin::CameraConnect(const FMagicLeapCameraConnect& ResultDelegate)
+{
+	if (TryPushNewCaptureTask(FCameraTask::EType::Connect))
+	{
+		OnCameraConnect = ResultDelegate;
+		return true;
+	}
+
+	return false;
+}
+
+bool FMagicLeapCameraPlugin::CameraDisconnect(const FMagicLeapCameraDisconnect& ResultDelegate)
+{
+	if (UserCount <= 0 && TryPushNewCaptureTask(FCameraTask::EType::Disconnect))
+	{
+		OnCameraDisconnect = ResultDelegate;
+		return true;
+	}
+
+	return false;
+}
+
+int64 FMagicLeapCameraPlugin::GetPreviewHandle() const
+{
+	return FCameraRunnable::PreviewHandle.GetValue();
+}
+
+void FMagicLeapCameraPlugin::IncUserCount()
+{
+	++UserCount;
+}
+
+void FMagicLeapCameraPlugin::DecUserCount()
+{
+	--UserCount;
+	if (UserCount <= 0)
+	{
+		UserCount = 0;
+		TryPushNewCaptureTask(FCameraTask::EType::Disconnect);
+	}
+}
+
+bool FMagicLeapCameraPlugin::CaptureImageToFileAsync(const FMagicLeapCameraCaptureImgToFileMulti& ResultDelegate)
+{
+	if (TryPushNewCaptureTask(FCameraTask::EType::ImageToFile))
+	{
+		OnCaptureImgToFile = ResultDelegate;
+		return true;
+	}
+
+	return false;
+}
+
+bool FMagicLeapCameraPlugin::CaptureImageToTextureAsync(const FMagicLeapCameraCaptureImgToTextureMulti& ResultDelegate)
+{
+	if (TryPushNewCaptureTask(FCameraTask::EType::ImageToTexture))
+	{
+		OnCaptureImgToTexture = ResultDelegate;
+		return true;
+	}
+
+	return false;
+}
+
+bool FMagicLeapCameraPlugin::StartRecordingAsync(const FMagicLeapCameraStartRecordingMulti& ResultDelegate)
+{
+	if (TryPushNewCaptureTask(FCameraTask::EType::StartVideoToFile))
+	{
+		OnStartRecording = ResultDelegate;
+		return true;
+	}
+
+	return false;
+}
+
+bool FMagicLeapCameraPlugin::StopRecordingAsync(const FMagicLeapCameraStopRecordingMulti& ResultDelegate)
+{
+	if (TryPushNewCaptureTask(FCameraTask::EType::StopVideoToFile))
+	{
+		OnStopRecording = ResultDelegate;
+		return true;
+	}
+
+	return false;
+}
+
+bool FMagicLeapCameraPlugin::SetLogDelegate(const FMagicLeapCameraLogMessageMulti& LogDelegate)
+{
+	OnLogMessage = LogDelegate;
+	return true;
+}
+
+bool FMagicLeapCameraPlugin::IsCapturing() const
+{
+	return CurrentTaskType != FCameraTask::EType::None;
+}
+
+bool FMagicLeapCameraPlugin::TryPushNewCaptureTask(FCameraTask::EType InTaskType)
+{
+	bool bCanPushTask = false;
+
+	switch (InTaskType)
+	{
+	case FCameraTask::EType::None:
+	{
+		bCanPushTask = false;
+	}
+	break;
+
+	case FCameraTask::EType::Connect:
+	{
+		bCanPushTask = CurrentTaskType == FCameraTask::EType::None || CurrentTaskType == FCameraTask::EType::Disconnect;
+	}
+	break;
+
+	case FCameraTask::EType::Disconnect:
+	{
+		bCanPushTask = CurrentTaskType != FCameraTask::EType::Disconnect;
+	}
+	break;
+
+	case FCameraTask::EType::ImageToFile:
+	{
+		bCanPushTask = CurrentTaskType == FCameraTask::EType::None;
+	}
+	break;
+
+	case FCameraTask::EType::ImageToTexture:
+	{
+		bCanPushTask = CurrentTaskType == FCameraTask::EType::None;
+	}
+	break;
+
+	case FCameraTask::EType::StartVideoToFile:
+	{
+		bCanPushTask = CurrentTaskType == FCameraTask::EType::None && PrevTaskType != FCameraTask::EType::StartVideoToFile;
+	}
+	break;
+
+	case FCameraTask::EType::StopVideoToFile:
+	{
+		bCanPushTask = PrevTaskType != FCameraTask::EType::StopVideoToFile &&
+			(CurrentTaskType == FCameraTask::EType::None || CurrentTaskType == FCameraTask::EType::StartVideoToFile);
+	}
+	break;
+
+	case FCameraTask::EType::Log:
+	{
+		bCanPushTask = true;
+	}
+	break;
+	}
+
+	if (bCanPushTask)
+	{
+		if (InTaskType != FCameraTask::EType::Log)
+		{
+			CurrentTaskType = InTaskType;
+		}
+
+		Runnable->PushNewCaptureTask(InTaskType);
+		return true;
+	}
+
+	return false;
+}
+
+void FMagicLeapCameraPlugin::OnAppPause()
+{
+	// Cancel the current video recording (if one is active).
+	if (CurrentTaskType == FCameraTask::EType::StartVideoToFile)
+	{
+		PrevTaskType = CurrentTaskType;
+		CurrentTaskType = FCameraTask::EType::StopVideoToFile;
+		// The runnable will take care of terminating the video on it's own.
+	}
+}
diff --git a/Engine/Plugins/Lumin/MagicLeapCamera/Source/MagicLeapCamera/Private/MagicLeapCameraPlugin.h b/Engine/Plugins/Lumin/MagicLeapCamera/Source/MagicLeapCamera/Private/MagicLeapCameraPlugin.h
new file mode 100644
--- /dev/null
+++ b/Engine/Plugins/Lumin/MagicLeapCamera/Source/MagicLeapCamera/Private/MagicLeapCameraPlugin.h
@@ -0,0 +1,59 @@
+// Copyright 1998-2019 Epic Games, Inc. All Rights Reserved.
+
+#pragma once
+
+#include "CoreMinimal.h"
+#include "IMagicLeapPlugin.h"
+#include "IMagicLeapCameraPlugin.h"
+#include "MagicLeapCameraRunnable.h"
+#include "MagicLeapCameraTypes.h"
+#include "MagicLeapPluginUtil.h"
+#include "AppEventHandler.h"
+
+DECLARE_LOG_CATEGORY_EXTERN(LogMagicLeapCamera, Verbose, All);
+
+class FMagicLeapCameraPlugin : public IMagicLeapCameraPlugin
+{
+public:
+	FMagicLeapCameraPlugin();
+	void StartupModule() override;
+	void ShutdownModule() override;
+	bool Tick(float DeltaTime) override;
+	bool CameraConnect(const FMagicLeapCameraConnect& ResultDelegate) override;
+	bool CameraDisconnect(const FMagicLeapCameraDisconnect& ResultDelegate) override;
+	int64 GetPreviewHandle() const override;
+
+	void IncUserCount();
+	void DecUserCount();
+	bool CaptureImageToFileAsync(const FMagicLeapCameraCaptureImgToFileMulti& ResultDelegate);
+	bool CaptureImageToTextureAsync(const FMagicLeapCameraCaptureImgToTextureMulti& ResultDelegate);
+	bool StartRecordingAsync(const FMagicLeapCameraStartRecordingMulti& ResultDelegate);
+	bool StopRecordingAsync(const FMagicLeapCameraStopRecordingMulti& ResultDelegate);
+	bool SetLogDelegate(const FMagicLeapCameraLogMessageMulti& LogDelegate);
+	bool IsCapturing() const;
+
+private:
+	bool TryPushNewCaptureTask(FCameraTask::EType InTaskType);
+	void OnAppPause();
+
+	FMagicLeapAPISetup APISetup;
+	FTickerDelegate TickDelegate;
+	FDelegateHandle TickDelegateHandle;
+	uint32 UserCount;
+	MagicLeap::IAppEventHandler AppEventHandler;
+	FCameraRunnable* Runnable;
+	FCameraTask::EType CurrentTaskType;
+	FCameraTask::EType PrevTaskType;
+	FMagicLeapCameraConnect OnCameraConnect;
+	FMagicLeapCameraDisconnect OnCameraDisconnect;
+	FMagicLeapCameraCaptureImgToFileMulti OnCaptureImgToFile;
+	FMagicLeapCameraCaptureImgToTextureMulti OnCaptureImgToTexture;
+	FMagicLeapCameraStartRecordingMulti OnStartRecording;
+	FMagicLeapCameraStopRecordingMulti OnStopRecording;
+	FMagicLeapCameraLogMessageMulti OnLogMessage;
+};
+
+inline FMagicLeapCameraPlugin& GetMagicLeapCameraPlugin()
+{
+	return FModuleManager::Get().GetModuleChecked<FMagicLeapCameraPlugin>("MagicLeapCamera");
+}
diff --git a/Engine/Plugins/Lumin/MagicLeapCamera/Source/MagicLeapCamera/Private/MagicLeapCameraRunnable.cpp b/Engine/Plugins/Lumin/MagicLeapCamera/Source/MagicLeapCamera/Private/MagicLeapCameraRunnable.cpp
new file mode 100644
--- /dev/null
+++ b/Engine/Plugins/Lumin/MagicLeapCamera/Source/MagicLeapCamera/Private/MagicLeapCameraRunnable.cpp
@@ -0,0 +1,356 @@
+// Copyright 1998-2019 Epic Games, Inc. All Rights Reserved.
+
+#include "MagicLeapCameraRunnable.h"
+#include "MagicLeapCameraPlugin.h"
+#include "Engine/Texture2D.h"
+#include "Lumin/LuminPlatformFile.h"
+#include "Lumin/CAPIShims/LuminAPIMediaError.h"
+
+#if WITH_MLSDK
+FThreadSafeCounter64 FCameraRunnable::PreviewHandle = ML_INVALID_HANDLE;
+#else
+FThreadSafeCounter64 FCameraRunnable::PreviewHandle = 0;
+#endif //WITH_MLSDK
+
+FCameraRunnable::FCameraRunnable()
+	: FMagicLeapRunnable({ EMagicLeapPrivilege::CameraCapture, EMagicLeapPrivilege::AudioCaptureMic, EMagicLeapPrivilege::VoiceInput }, TEXT("FCameraRunnable"))
+	, ThreadSafeConnectionStatus(static_cast<int32>(EConnectionStatus::NotConnected))
+	, bWasConnectedOnPause(false)
+	, ImgExtension(".jpeg")
+	, VidExtension(".mp4")
+{
+	IImageWrapperModule& ImageWrapperModule = FModuleManager::LoadModuleChecked<IImageWrapperModule>(FName("ImageWrapper"));
+	ImageWrapper = ImageWrapperModule.CreateImageWrapper(EImageFormat::JPEG);
+}
+
+void FCameraRunnable::Exit()
+{
+#if WITH_MLSDK
+	TryDisconnect();
+#endif // WITH_MLSDK
+}
+
+void FCameraRunnable::PushNewCaptureTask(FCameraTask::EType InTaskType)
+{
+#if WITH_MLSDK
+	if (InTaskType == FCameraTask::EType::Disconnect && GetConnectionStatus() == EConnectionStatus::NotConnected)
+	{
+		return;
+	}
+
+	if (InTaskType == FCameraTask::EType::Connect && GetConnectionStatus() != EConnectionStatus::NotConnected)
+	{
+		return;
+	}
+
+	if (InTaskType != FCameraTask::EType::Connect &&
+		InTaskType != FCameraTask::EType::Disconnect &&
+		GetConnectionStatus() == EConnectionStatus::NotConnected)
+	{
+		FCameraTask ConnectTask;
+		ConnectTask.CaptureType = FCameraTask::EType::Connect;
+		PushNewTask(ConnectTask);
+	}
+	FCameraTask CaptureImgToFileTask;
+	CaptureImgToFileTask.CaptureType = InTaskType;
+	PushNewTask(CaptureImgToFileTask);
+#endif // WITH_MLSDK
+}
+
+bool FCameraRunnable::IsConnected() const
+{
+	return GetConnectionStatus() == EConnectionStatus::Connected;
+}
+
+void FCameraRunnable::Pause()
+{
+#if WITH_MLSDK
+	bWasConnectedOnPause = IsConnected();
+	// Cancel the current video recording (if one is active).
+	if (CurrentTask.CaptureType == FCameraTask::EType::StartVideoToFile)
+	{
+		StopRecordingVideo();
+		CurrentTask.bSuccess = false;
+		PushCompletedTask(CurrentTask);
+	}
+	// Cancel any incoming tasks.
+	CancelIncomingTasks();
+	// Disconnect camera if connected
+	TryDisconnect();
+#endif // WITH_MLSDK
+}
+
+void FCameraRunnable::Resume()
+{
+	if (bWasConnectedOnPause)
+	{
+		FCameraTask ConnectTask;
+		ConnectTask.CaptureType = FCameraTask::EType::Connect;
+		PushNewTask(ConnectTask);
+	}
+}
+
+bool FCameraRunnable::ProcessCurrentTask()
+{
+	bool bSuccess = false;
+
+#if WITH_MLSDK
+	switch (CurrentTask.CaptureType)
+	{
+	case FCameraTask::EType::None: bSuccess = false; checkf(false, TEXT("Invalid camera task encountered!")); break;
+	case FCameraTask::EType::Connect:
+	{
+		SetConnectionStatus(EConnectionStatus::Connecting);
+		bSuccess = TryConnect();
+		SetConnectionStatus(bSuccess ? EConnectionStatus::Connected : EConnectionStatus::NotConnected);
+	}
+	break;
+	case FCameraTask::EType::Disconnect: bSuccess = TryDisconnect(); break;
+	case FCameraTask::EType::ImageToFile: bSuccess = CaptureImageToFile(); break;
+	case FCameraTask::EType::ImageToTexture: bSuccess = CaptureImageToTexture(); break;
+	case FCameraTask::EType::StartVideoToFile: bSuccess = StartRecordingVideo(); break;
+	case FCameraTask::EType::StopVideoToFile: bSuccess = StopRecordingVideo(); break;
+	}
+#endif // WITH_MLSDK
+	return bSuccess;
+}
+
+#if WITH_MLSDK
+void FCameraRunnable::OnPreviewBufferAvailable(MLHandle Output, void *Data)
+{
+	(void)Data;
+	PreviewHandle.Set(static_cast<int64>(Output));
+}
+
+bool FCameraRunnable::TryConnect()
+{
+	//IncomingTasks.Peek(CurrentTask); // this is purely so that any log messages go to the correct requester
+	if (AppEventHandler.GetPrivilegeStatus(EMagicLeapPrivilege::CameraCapture) != MagicLeap::EPrivilegeState::Granted)
+	{
+		Log(TEXT("Cannot connect to camera due to lack of privilege!"));
+		return false;
+	}
+	else
+	{
+		if (bPaused) return false;
+
+		MLResult Result = MLCameraConnect();
+
+		if (bPaused) return false;
+
+		if (Result != MLResult_Ok)
+		{
+			Log(FString::Printf(TEXT("MLCameraConnect failed with error %s!"), UTF8_TO_TCHAR(MLMediaResultGetString(Result))));
+			CancelIncomingTasks();
+			return false;
+		}
+		else
+		{
+			FMemory::Memset(&DeviceStatusCallbacks, 0, sizeof(MLCameraDeviceStatusCallbacks));
+			DeviceStatusCallbacks.on_preview_buffer_available = OnPreviewBufferAvailable;
+			Result = MLCameraSetDeviceStatusCallbacks(&DeviceStatusCallbacks, nullptr);
+			if (Result != MLResult_Ok)
+			{
+				Log(FString::Printf(TEXT("MLCameraSetDeviceStatusCallbacks failed with error %s!"), UTF8_TO_TCHAR(MLMediaResultGetString(Result))));
+			}
+		}
+	}
+
+	return true;
+}
+
+bool FCameraRunnable::TryDisconnect()
+{
+	if (IsConnected())
+	{
+		MLResult Result = MLCameraDisconnect();
+		if (Result != MLResult_Ok)
+		{
+			Log(FString::Printf(TEXT("MLCameraDisconnect failed with error %s!"), UTF8_TO_TCHAR(MLMediaResultGetString(Result))));
+		}
+		else
+		{
+			SetConnectionStatus(EConnectionStatus::NotConnected);
+		}
+	}
+
+	return !IsConnected();
+}
+
+bool FCameraRunnable::CaptureImageToFile()
+{
+	if (bPaused) return false;
+
+	Log(TEXT("Beginning capture image to file."));
+	MLHandle Handle = ML_INVALID_HANDLE;
+	MLResult Result = MLCameraPrepareCapture(MLCameraCaptureType_Image, &Handle);
+	if (Result != MLResult_Ok)
+	{
+		Log(FString::Printf(TEXT("MLCameraPrepareCapture failed with error %s!  Camera capture aborted!"), UTF8_TO_TCHAR(MLMediaResultGetString(Result))));
+		return false;
+	}
+
+	if (bPaused) return false;
+
+#if PLATFORM_LUMIN
+	IPlatformFile& PlatformFile = IPlatformFile::GetPlatformPhysical();
+	// This module is only for Lumin so this is fine for now.
+	FLuminPlatformFile* LuminPlatformFile = static_cast<FLuminPlatformFile*>(&PlatformFile);
+	UniqueFileName = LuminPlatformFile->ConvertToLuminPath(FPaths::CreateTempFilename(*FPaths::ProjectSavedDir(), TEXT("Img_"), *ImgExtension), true);
+#endif
+	Result = MLCameraCaptureImage(TCHAR_TO_UTF8(*UniqueFileName));
+	if (Result != MLResult_Ok)
+	{
+		Log(FString::Printf(TEXT("MLCameraCaptureImage failed with error %s!  Camera capture aborted!"), UTF8_TO_TCHAR(MLMediaResultGetString(Result))));
+		return false;
+	}
+
+	Log(FString::Printf(TEXT("Captured image to %s"), *UniqueFileName));
+	CurrentTask.FilePath = UniqueFileName;
+	return true;
+}
+
+bool FCameraRunnable::CaptureImageToTexture()
+{
+	if (bPaused) return false;
+
+	Log(TEXT("Beginning capture image to texture."));
+	MLCameraOutput* CameraOutput = nullptr;
+	MLHandle Handle = ML_INVALID_HANDLE;
+	MLResult Result = MLCameraPrepareCapture(MLCameraCaptureType_ImageRaw, &Handle);
+	if (Result != MLResult_Ok)
+	{
+		Log(FString::Printf(TEXT("MLCameraPrepareCapture failed with error %s!  Camera capture aborted!"), UTF8_TO_TCHAR(MLMediaResultGetString(Result))));
+		return false;
+	}
+
+	if (bPaused) return false;
+
+	Result = MLCameraCaptureImageRaw();
+	if (Result != MLResult_Ok)
+	{
+		Log(FString::Printf(TEXT("MLCameraCaptureImageRaw failed with error %s!  Camera capture aborted!"), UTF8_TO_TCHAR(MLMediaResultGetString(Result))));
+		return false;
+	}
+
+	if (bPaused) return false;
+
+	Result = MLCameraGetImageStream(&CameraOutput);
+	if (Result != MLResult_Ok)
+	{
+		Log(FString::Printf(TEXT("MLCameraGetImageStream failed with error %s!  Camera capture aborted!"), UTF8_TO_TCHAR(MLMediaResultGetString(Result))));
+		return false;
+	}
+
+	if (bPaused) return false;
+
+	if (CameraOutput->plane_count == 0)
+	{
+		Log(TEXT("Invalid plane_count!  Camera capture aborted!"));
+		return false;
+	}
+
+	MLCameraPlaneInfo& ImageInfo = CameraOutput->planes[0];
+	if (ImageWrapper.IsValid() && ImageWrapper->SetCompressed(ImageInfo.data, ImageInfo.size))
+	{
+		const TArray<uint8>* RawData = NULL;
+		if (ImageWrapper->GetRaw(ImageWrapper->GetFormat(), 8, RawData))
+		{
+			Log(FString::Printf(TEXT("ImageWrapper width=%d height=%d size=%d"), ImageWrapper->GetWidth(), ImageWrapper->GetHeight(), RawData->Num()));
+			UTexture2D* CaptureTexture = UTexture2D::CreateTransient(ImageWrapper->GetWidth(), ImageWrapper->GetHeight(), EPixelFormat::PF_R8G8B8A8);
+			CaptureTexture->AddToRoot();
+			FTexture2DMipMap& Mip = CaptureTexture->PlatformData->Mips[0];
+			void* Data = Mip.BulkData.Lock(LOCK_READ_WRITE);
+			FMemory::Memcpy(Data, RawData->GetData(), Mip.BulkData.GetBulkDataSize());
+			Mip.BulkData.Unlock();
+			CaptureTexture->UpdateResource();
+			CurrentTask.Texture = CaptureTexture;
+		}
+	}
+
+	return true;
+}
+
+bool FCameraRunnable::StartRecordingVideo()
+{
+	if (bPaused) return false;
+
+	Log(TEXT("Beginning capture video to file."));
+	MLHandle Handle = ML_INVALID_HANDLE;
+	MLResult Result = MLCameraPrepareCapture(MLCameraCaptureType_Video, &Handle);
+	if (Result != MLResult_Ok)
+	{
+		Log(FString::Printf(TEXT("MLCameraPrepareCapture failed with error %s!  Camera capture aborted!"), UTF8_TO_TCHAR(MLMediaResultGetString(Result))));
+		return false;
+	}
+
+	if (bPaused) return false;
+
+	if (AppEventHandler.GetPrivilegeStatus(EMagicLeapPrivilege::AudioCaptureMic) != MagicLeap::EPrivilegeState::Granted)
+	{
+		Log(TEXT("Cannot capture video due to lack of privilege!"));
+		return false;
+	}
+
+	if (bPaused) return false;
+
+	if (AppEventHandler.GetPrivilegeStatus(EMagicLeapPrivilege::VoiceInput) != MagicLeap::EPrivilegeState::Granted)
+	{
+		Log(TEXT("Cannot capture video due to lack of privilege!"));
+		return false;
+	}
+
+	if (bPaused) return false;
+
+#if PLATFORM_LUMIN
+	IPlatformFile& PlatformFile = IPlatformFile::GetPlatformPhysical();
+	// This module is only for Lumin so this is fine for now.
+	FLuminPlatformFile* LuminPlatformFile = static_cast<FLuminPlatformFile*>(&PlatformFile);
+	UniqueFileName = LuminPlatformFile->ConvertToLuminPath(FPaths::CreateTempFilename(*FPaths::ProjectSavedDir(), TEXT("Vid_"), *VidExtension), true);
+#endif
+	Result = MLCameraCaptureVideoStart(TCHAR_TO_UTF8(*UniqueFileName));
+	if (Result != MLResult_Ok)
+	{
+		Log(FString::Printf(TEXT("MLCameraCaptureVideoStart failed with error %s!  Video capture aborted!"), UTF8_TO_TCHAR(MLMediaResultGetString(Result))));
+		return false;
+	}
+
+	return true;
+}
+
+bool FCameraRunnable::StopRecordingVideo()
+{
+	MLResult Result = MLCameraCaptureVideoStop();
+	if (Result != MLResult_Ok)
+	{
+		Log(FString::Printf(TEXT("MLCameraCaptureVideoStop failed with error %s!"), UTF8_TO_TCHAR(MLMediaResultGetString(Result))));
+		return false;
+	}
+	else
+	{
+		Log(FString::Printf(TEXT("Captured video to %s"), *UniqueFileName));
+	}
+
+	CurrentTask.FilePath = UniqueFileName;
+	return true;
+}
+
+void FCameraRunnable::Log(const FString& Info)
+{
+	FCameraTask LogTask;
+	LogTask.CaptureType = FCameraTask::EType::Log;
+	LogTask.Log = Info;
+	PushCompletedTask(LogTask);
+	UE_LOG(LogMagicLeapCamera, Log, TEXT("%s"), *Info);
+}
+#endif //WITH_MLSDK
+
+void FCameraRunnable::SetConnectionStatus(EConnectionStatus ConnectionStatus)
+{
+	ThreadSafeConnectionStatus.Set(static_cast<int32>(ConnectionStatus));
+}
+
+FCameraRunnable::EConnectionStatus FCameraRunnable::GetConnectionStatus() const
+{
+	return static_cast<EConnectionStatus>(ThreadSafeConnectionStatus.GetValue());
+}
diff --git a/Engine/Plugins/Lumin/MagicLeapCamera/Source/MagicLeapCamera/Private/MagicLeapCameraRunnable.h b/Engine/Plugins/Lumin/MagicLeapCamera/Source/MagicLeapCamera/Private/MagicLeapCameraRunnable.h
new file mode 100644
--- /dev/null
+++ b/Engine/Plugins/Lumin/MagicLeapCamera/Source/MagicLeapCamera/Private/MagicLeapCameraRunnable.h
@@ -0,0 +1,81 @@
+// Copyright 1998-2019 Epic Games, Inc. All Rights Reserved.
+
+#pragma once
+
+#include "MagicLeapRunnable.h"
+#include "HAL/ThreadSafeCounter.h"
+#include "IImageWrapper.h"
+#include "IImageWrapperModule.h"
+#include "Lumin/CAPIShims/LuminAPICamera.h"
+
+struct FCameraTask : public FMagicLeapTask
+{
+	enum class EType : uint32
+	{
+		None,
+		Connect,
+		Disconnect,
+		ImageToFile,
+		ImageToTexture,
+		StartVideoToFile,
+		StopVideoToFile,
+		Log,
+	};
+
+	EType CaptureType;
+	FString FilePath;
+	FString Log;
+	class UTexture2D* Texture;
+
+	FCameraTask()
+		: CaptureType(EType::None)
+		, Texture(nullptr)
+	{
+	}
+};
+
+class FCameraRunnable : public FMagicLeapRunnable<FCameraTask>
+{
+public:
+	FCameraRunnable();
+	void Exit() override;
+	void PushNewCaptureTask(FCameraTask::EType InTaskType);
+	bool IsConnected() const;
+
+	static FThreadSafeCounter64 PreviewHandle;
+
+protected:
+	void Pause() override;
+	void Resume() override;
+
+private:
+	bool ProcessCurrentTask() override;
+#if WITH_MLSDK
+	static void OnPreviewBufferAvailable(MLHandle Output, void *Data);
+	bool TryConnect();
+	bool TryDisconnect();
+	bool CaptureImageToFile();
+	bool CaptureImageToTexture();
+	bool StartRecordingVideo();
+	bool StopRecordingVideo();
+	void Log(const FString& InLogMsg);
+
+	MLCameraDeviceStatusCallbacks DeviceStatusCallbacks;
+#endif //WITH_MLSDK
+	FCameraTask::EType CurrentTaskType;
+	enum class EConnectionStatus : int32
+	{
+		NotConnected,
+		Connecting,
+		Connected
+	};
+	FThreadSafeCounter ThreadSafeConnectionStatus;
+	bool bWasConnectedOnPause;
+	const FString ImgExtension;
+	const FString VidExtension;
+	FString UniqueFileName;
+	TSharedPtr<IImageWrapper> ImageWrapper;
+
+	void SetConnectionStatus(EConnectionStatus ConnectionStatus);
+	EConnectionStatus GetConnectionStatus() const;
+};
diff --git a/Engine/Plugins/Lumin/MagicLeapCamera/Source/MagicLeapCamera/Public/IMagicLeapCameraPlugin.h b/Engine/Plugins/Lumin/MagicLeapCamera/Source/MagicLeapCamera/Public/IMagicLeapCameraPlugin.h
new file mode 100644
--- /dev/null
+++ b/Engine/Plugins/Lumin/MagicLeapCamera/Source/MagicLeapCamera/Public/IMagicLeapCameraPlugin.h
@@ -0,0 +1,41 @@
+// Copyright 1998-2019 Epic Games, Inc. All Rights Reserved.
+
+#pragma once
+
+#include "Modules/ModuleManager.h"
+#include "Containers/Ticker.h"
+#include "MagicLeapCameraTypes.h"
+
+/**
+ * The public interface to this module.  In most cases, this interface is only public to sibling modules
+ * within this plugin.
+ */
+class IMagicLeapCameraPlugin : public IModuleInterface
+{
+public:
+	/**
+	 * Singleton-like access to this module's interface.  This is just for convenience!
+	 * Beware of calling this during the shutdown phase, though.  Your module might have been unloaded already.
+	 *
+	 * @return Returns singleton instance, loading the module on demand if needed
+	 */
+	static inline IMagicLeapCameraPlugin& Get()
+	{
+		return FModuleManager::LoadModuleChecked<IMagicLeapCameraPlugin>("MagicLeapCamera");
+	}
+
+	/**
+	 * Checks to see if this module is loaded and ready.  It is only valid to call Get() if IsAvailable() returns true.
+	 *
+	 * @return True if the module is loaded and ready to use
+	 */
+	static inline bool IsAvailable()
+	{
+		return FModuleManager::Get().IsModuleLoaded("MagicLeapCamera");
+	}
+
+	virtual bool Tick(float DeltaTime) = 0;
+	virtual bool CameraConnect(const FMagicLeapCameraConnect& ResultDelegate) { return true; }
+	virtual bool CameraDisconnect(const FMagicLeapCameraDisconnect& ResultDelegate) { return true; }
+	virtual int64 GetPreviewHandle() const = 0;
+};
diff --git a/Engine/Plugins/Lumin/MagicLeapCamera/Source/MagicLeapCamera/Public/MagicLeapCameraComponent.h b/Engine/Plugins/Lumin/MagicLeapCamera/Source/MagicLeapCamera/Public/MagicLeapCameraComponent.h
new file mode 100644
--- /dev/null
+++ b/Engine/Plugins/Lumin/MagicLeapCamera/Source/MagicLeapCamera/Public/MagicLeapCameraComponent.h
@@ -0,0 +1,86 @@
+// Copyright 1998-2019 Epic Games, Inc. All Rights Reserved.
+
+#pragma once
+
+#include "Engine/Engine.h"
+#include "Components/ActorComponent.h"
+#include "MagicLeapCameraTypes.h"
+#include "MagicLeapCameraComponent.generated.h"
+
+/**
+  The MagicLeapCameraComponent provides access to and maintains state for camera capture functionality.
+  The connection to the device's camera is managed internally.  Users of this component
+  are able to asynchronously capture camera images and footage to file.  Alternatively,
+  a camera image can be captured directly to texture.  The user need only make the relevant
+  asynchronous call and then register the appropriate event handlers for the
+  operation's completion.
+*/
+UCLASS(ClassGroup = MagicLeap, BlueprintType, Blueprintable, EditInlineNew, meta = (BlueprintSpawnableComponent))
+class MAGICLEAPCAMERA_API UMagicLeapCameraComponent : public UActorComponent
+{
+	GENERATED_BODY()
+
+public:
+	/** Notifies the MagicLeapCameraPlugin of a new user. */
+	void BeginPlay() override;
+
+	/** Notifies the MagicLeapCameraPlugin that a user is being destroyed (needed for auto-disconnecting the camera). */
+	void EndPlay(const EEndPlayReason::Type EndPlayReason) override;
+
+	/**
+		Initiates a capture image to file task on a separate thread.
+		@brief The newly created jpeg file will have an automatically generated name which is guaranteed
+			   to be unique.  Upon completion, a successful operation will provide the file path of the newly
+			   created jpeg to the FMagicLeapCameraCaptureImgToFile event handler.
+		@return True if the call succeeds, false otherwise.
+	*/
+	UFUNCTION(BlueprintCallable, Category = "Camera | MagicLeap")
+	virtual bool CaptureImageToFileAsync();
+
+	/**
+		Initiates a capture image to memory task on a speparate thread.
+		@brief The user should register event handlers for both the success and fail events.  Upon completion,
+			   a successful operation will provide a dynamically generated texture containing the captured
+			   image to the FMagicLeapCameraCaptureImgToTextureSuccess event handler.
+		@note The generated texture will be garbage collected when this app is destroyed.
+		@return True if the call succeeds, false otherwise.
+	*/
+	UFUNCTION(BlueprintCallable, Category = "Camera | MagicLeap")
+	virtual bool CaptureImageToTextureAsync();
+
+	/**
+		Initiates the capturing of video/audio data on a separate thread.
+		@note The system will continue to record video until StopRecordingVideo is called.
+		@return True if the call succeeds, false otherwise.
+	*/
+	UFUNCTION(BlueprintCallable, Category = "Camera | MagicLeap")
+	virtual bool StartRecordingAsync();
+
+	/**
+		Stops the recording and saves the video/audio data to an mp4 file.
+		@note The newly created mp4 file will have an automatically generated name which is guaranteed
+			  to be unique.
+		@return True if the call succeeds, false otherwise.
+	*/
+	UFUNCTION(BlueprintCallable, Category = "Camera | MagicLeap")
+	virtual bool StopRecordingAsync();
+
+	/**
+		Gets the capture state of the component.
+		@return True if the component is currently capturing, false otherwise.
+	*/
+	UFUNCTION(BlueprintPure, Category = "Camera | MagicLeap")
+	bool IsCapturing() const;
+
+private:
+	UPROPERTY(BlueprintAssignable, Category = "Camera | MagicLeap", meta = (AllowPrivateAccess = true))
+	FMagicLeapCameraCaptureImgToFileMulti OnCaptureImgToFile;
+	UPROPERTY(BlueprintAssignable, Category = "Camera | MagicLeap", meta = (AllowPrivateAccess = true))
+	FMagicLeapCameraCaptureImgToTextureMulti OnCaptureImgToTexture;
+	UPROPERTY(BlueprintAssignable, Category = "Camera | MagicLeap", meta = (AllowPrivateAccess = true))
+	FMagicLeapCameraStartRecordingMulti OnStartRecording;
+	UPROPERTY(BlueprintAssignable, Category = "Camera | MagicLeap", meta = (AllowPrivateAccess = true))
+	FMagicLeapCameraStopRecordingMulti OnStopRecording;
+	UPROPERTY(BlueprintAssignable, Category = "Camera | MagicLeap", meta = (AllowPrivateAccess = true))
+	FMagicLeapCameraLogMessageMulti OnLogMessage;
+};
diff --git a/Engine/Plugins/Lumin/MagicLeapCamera/Source/MagicLeapCamera/Public/MagicLeapCameraFunctionLibrary.h b/Engine/Plugins/Lumin/MagicLeapCamera/Source/MagicLeapCamera/Public/MagicLeapCameraFunctionLibrary.h
new file mode 100644
--- /dev/null
+++ b/Engine/Plugins/Lumin/MagicLeapCamera/Source/MagicLeapCamera/Public/MagicLeapCameraFunctionLibrary.h
@@ -0,0 +1,99 @@
+// Copyright 1998-2019 Epic Games, Inc. All Rights Reserved.
+#pragma once
+
+#include "Engine/Engine.h"
+#include "Kismet/BlueprintFunctionLibrary.h"
+#include "MagicLeapCameraTypes.h"
+#include "MagicLeapCameraFunctionLibrary.generated.h"
+
+/**
+  The MagicLeapCameraFunctionLibrary provides access to the camera capture functionality.
+  Users of this library are able to asynchronously capture camera images and footage to file.
+  Alternatively, a camera image can be captured directly to texture.  The user need only make
+  the relevant asynchronous call and then register the appropriate event handlers for the
+  operation's completion.
+*/
+UCLASS(ClassGroup = MagicLeap)
+class MAGICLEAPCAMERA_API UMagicLeapCameraFunctionLibrary : public UBlueprintFunctionLibrary
+{
+	GENERATED_BODY()
+
+public:
+	/** 
+		Establishes a connection with the device's camera.
+		@note A connection will be made automatically upon the first capture call if this is
+			  not called first.  Calling this function manually allows the developer to control
+			  when privilege notifications for this plugin will be activated (if application
+			  is being used for the first time).
+		@return True if the call succeeds, false otherwise.
+	*/
+	UFUNCTION(BlueprintCallable, Category = "Camera Function Library | MagicLeap")
+	static bool CameraConnect(const FMagicLeapCameraConnect& ResultDelegate);
+
+	/** 
+		Disconnects from the device's camera.
+		@note This function must be called before the application terminates (if the camera has
+			  been connected to).  Failure to do so will result in the camera connection remaining
+			  open (and the camera icon remaining on screen).
+		@return True if the call succeeds, false otherwise.
+	*/
+	UFUNCTION(BlueprintCallable, Category = "Camera Function Library | MagicLeap")
+	static bool CameraDisconnect(const FMagicLeapCameraDisconnect& ResultDelegate);
+
+	/**
+		Initiates a capture image to file task on a separate thread.
+		@brief The newly created jpeg file will have an automatically generated name which is guaranteed
+			   to be unique.  Upon completion, a successful operation will provide the file path of the newly
+			   created jpeg to the FMagicLeapCameraCaptureImgToFile event handler.
+		@param ResultDelegate The delegate to be notified once the camera image has been saved to a jpeg file.
+		@return True if the call succeeds, false otherwise.
+	*/
+	UFUNCTION(BlueprintCallable, Category = "Camera Function Library | MagicLeap")
+	static bool CaptureImageToFileAsync(const FMagicLeapCameraCaptureImgToFile& ResultDelegate);
+
+	/**
+		Initiates a capture image to memory task on a speparate thread.
+		@brief The user should register event handlers for both the success and fail events.  Upon completion,
+			   a successful operation will provide a dynamically generated texture containing the captured
+			   image to the FMagicLeapCameraCaptureImgToTextureSuccess event handler.
+		@note The generated texture will be garbage collected when this app is destroyed.
+		@param ResultDelegate The delegate to be notified once the camera image has been saved to a texture.
+		@return True if the call succeeds, false otherwise.
+	*/
+	UFUNCTION(BlueprintCallable, Category = "Camera Function Library | MagicLeap")
+	static bool CaptureImageToTextureAsync(const FMagicLeapCameraCaptureImgToTexture& ResultDelegate);
+
+	/**
+		Initiates the capturing of video/audio data on a separate thread.
+		@note The system will continue to record video until StopRecordingVideo is called.
+		@param ResultDelegate The delegate to be notified once the recording has begun or failed to begin.
+		@return True if the call succeeds, false otherwise.
+	*/
+	UFUNCTION(BlueprintCallable, Category = "Camera Function Library | MagicLeap")
+	static bool StartRecordingAsync(const FMagicLeapCameraStartRecording& ResultDelegate);
+
+	/**
+		Stops the recording and saves the video/audio data to an mp4 file.
+		@note The newly created mp4 file will have an automatically generated name which is guaranteed
+			  to be unique.
+		@param ResultDelegate The delegate to be notified once the video/audio data has been saved to an mp4 file.
+		@return True if the call succeeds, false otherwise.
+	*/
+	UFUNCTION(BlueprintCallable, Category = "Camera Function Library | MagicLeap")
+	static bool StopRecordingAsync(const FMagicLeapCameraStopRecording& ResultDelegate);
+
+	/**
+		Sets the delegate by which the system can pass log messages back to the calling blueprint.
+		@param LogDelegate The delegate by which the system will return log messages to the calling blueprint.
+		@return True if the call succeeds, false otherwise.
+	*/
+	UFUNCTION(BlueprintCallable, Category = "Camera Function Library | MagicLeap")
+	static bool SetLogDelegate(const FMagicLeapCameraLogMessage& LogDelegate);
+
+	/**
+		Gets the capture state of the component.
+		@return True if the component is currently capturing, false otherwise.
+	*/
+	UFUNCTION(BlueprintPure, Category = "Camera Function Library | MagicLeap")
+	static bool IsCapturing();
+};
diff --git a/Engine/Plugins/Lumin/MagicLeapCamera/Source/MagicLeapCamera/Public/MagicLeapCameraTypes.h b/Engine/Plugins/Lumin/MagicLeapCamera/Source/MagicLeapCamera/Public/MagicLeapCameraTypes.h
new file mode 100644
--- /dev/null
+++ b/Engine/Plugins/Lumin/MagicLeapCamera/Source/MagicLeapCamera/Public/MagicLeapCameraTypes.h
@@ -0,0 +1,57 @@
+// Copyright 1998-2019 Epic Games, Inc. All Rights Reserved.
+#pragma once
+
+#include "Engine/Engine.h"
+#include "MagicLeapCameraTypes.generated.h"
+
+USTRUCT()
+struct FMagicLeapCameraDummyStruct { GENERATED_BODY() };
+
+/** Delegate used to notify the initiating blueprint when the camera connect task has completed. */
+DECLARE_DYNAMIC_DELEGATE_OneParam(FMagicLeapCameraConnect, const bool, bSuccess);
+
+/** Delegate used to notify the initiating blueprint when the camera disonnect task has completed. */
+DECLARE_DYNAMIC_DELEGATE_OneParam(FMagicLeapCameraDisconnect, const bool, bSuccess);
+
+/**
+   Delegate used to notify the initiating blueprint when a capture image to file task has completed.
+   @note Although this signals the task as complete, it may have failed or been cancelled.
+   @param bSuccess True if the task succeeded, false otherwise.
+   @param FilePath A string containing the file path to the newly created jpeg.
+*/
+DECLARE_DYNAMIC_DELEGATE_TwoParams(FMagicLeapCameraCaptureImgToFile, const bool, bSuccess, const FString&, FilePath);
+DECLARE_DYNAMIC_MULTICAST_DELEGATE_TwoParams(FMagicLeapCameraCaptureImgToFileMulti, const bool, bSuccess, const FString&, FilePath);
+
+/**
+	Delegate used to pass the captured image back to the initiating blueprint.
+	@note The captured texture will remain in memory for the lifetime of the calling application (if the task succeeds).
+	@param bSuccess True if the task succeeded, false otherwise.
+	@param CaptureTexture A UTexture2D containing the captured image.
+*/
+DECLARE_DYNAMIC_DELEGATE_TwoParams(FMagicLeapCameraCaptureImgToTexture, const bool, bSuccess, UTexture2D*, CaptureTexture);
+DECLARE_DYNAMIC_MULTICAST_DELEGATE_TwoParams(FMagicLeapCameraCaptureImgToTextureMulti, const bool, bSuccess, UTexture2D*, CaptureTexture);
+
+/**
+	Delegate used to notify the initiating blueprint of the result of a request to begin recording video.
+	@note Although this signals the task as complete, it may have failed or been cancelled.
+	@param bSuccess True if the task succeeded, false otherwise.
+*/
+DECLARE_DYNAMIC_DELEGATE_OneParam(FMagicLeapCameraStartRecording, const bool, bSuccess);
+DECLARE_DYNAMIC_MULTICAST_DELEGATE_OneParam(FMagicLeapCameraStartRecordingMulti, const bool, bSuccess);
+
+/**
+	Delegate used to notify the initiating blueprint of the result of a request to stop recording video.
+	@note Although this signals the task as complete, it may have failed or been cancelled.
+	@param bSuccess True if the task succeeded, false otherwise.
+	@param FilePath A string containing the path to the newly created mp4.
+*/
+DECLARE_DYNAMIC_DELEGATE_TwoParams(FMagicLeapCameraStopRecording, const bool, bSuccess, const FString&, FilePath);
+DECLARE_DYNAMIC_MULTICAST_DELEGATE_TwoParams(FMagicLeapCameraStopRecordingMulti, const bool, bSuccess, const FString&, FilePath);
+
+/**
+	Delegate used to pass log messages from the capture worker thread to the initiating blueprint.
+	@note This is useful if the user wishes to have log messages in 3D space.
+	@param LogMessage A string containing the log message.
+*/
+DECLARE_DYNAMIC_DELEGATE_OneParam(FMagicLeapCameraLogMessage, const FString&, LogMessage);
+DECLARE_DYNAMIC_MULTICAST_DELEGATE_OneParam(FMagicLeapCameraLogMessageMulti, const FString&, LogMessage);
diff --git a/Engine/Plugins/Lumin/MagicLeapCamera/Source/MagicLeapCameraPreview/MagicLeapCameraPreview.Build.cs b/Engine/Plugins/Lumin/MagicLeapCamera/Source/MagicLeapCameraPreview/MagicLeapCameraPreview.Build.cs
new file mode 100644
--- /dev/null
+++ b/Engine/Plugins/Lumin/MagicLeapCamera/Source/MagicLeapCameraPreview/MagicLeapCameraPreview.Build.cs
@@ -0,0 +1,46 @@
+// Copyright 1998-2019 Epic Games, Inc. All Rights Reserved.
+
+using System.IO;
+
+namespace UnrealBuildTool.Rules
+{
+	public class MagicLeapCameraPreview : ModuleRules
+	{
+		public MagicLeapCameraPreview(ReadOnlyTargetRules Target) : base(Target)
+		{
+			PublicIncludePaths.Add(Path.Combine(ModuleDirectory, "Public"));
+
+			DynamicallyLoadedModuleNames.Add("Media");
+
+			PrivateDependencyModuleNames.AddRange(
+				new string[]
+				{
+					"Core",
+					"Engine",
+					"RenderCore",
+					"MediaUtils",
+					"LuminRuntimeSettings",
+					"MLSDK",
+					"MagicLeap",
+					"MagicLeapHelperVulkan",
+					"MagicLeapHelperOpenGL"
+				}
+			);
+
+			PrivateIncludePathModuleNames.AddRange(
+				new string[]
+				{
+					"MagicLeapCamera",
+					"Media"
+				}
+			);
+
+			PrivateIncludePaths.Add(Path.Combine(ModuleDirectory, "Private"));
+
+			if (Target.Platform == UnrealTargetPlatform.Lumin)
+			{
+				DynamicallyLoadedModuleNames.Add("MagicLeapCamera");
+			}
+		}
+	}
+}
diff --git a/Engine/Plugins/Lumin/MagicLeapCamera/Source/MagicLeapCameraPreview/Private/MagicLeapCameraPreviewModule.cpp b/Engine/Plugins/Lumin/MagicLeapCamera/Source/MagicLeapCameraPreview/Private/MagicLeapCameraPreviewModule.cpp
new file mode 100644
--- /dev/null
+++ b/Engine/Plugins/Lumin/MagicLeapCamera/Source/MagicLeapCameraPreview/Private/MagicLeapCameraPreviewModule.cpp
@@ -0,0 +1,31 @@
+// Copyright 1998-2019 Epic Games, Inc. All Rights Reserved.
+
+#include "MagicLeapCameraPreviewModule.h"
+#include "MagicLeapCameraPreviewPlayer.h"
+#include "Async/Async.h"
+
+#if 0
+PRAGMA_DISABLE_OPTIMIZATION
+#endif
+
+DEFINE_LOG_CATEGORY(LogMagicLeapCameraPreview);
+
+TSharedPtr<IMediaPlayer, ESPMode::ThreadSafe> FMagicLeapCameraPreviewModule::CreatePreviewPlayer(class IMediaEventSink& EventSink)
+{
+	// Disable camera preview on Vulkan until we are able to get it working.
+	// Can't call FLuminPlatformMisc::ShouldUseVulkan from here ...
+	bool bUseVulkan = false;
+	if (GConfig != nullptr)
+	{
+		GConfig->GetBool(TEXT("/Script/LuminRuntimeSettings.LuminRuntimeSettings"), TEXT("bUseVulkan"), bUseVulkan, GEngineIni);
+	}
+
+	if (!bUseVulkan)
+	{
+		return MakeShareable(new FMagicLeapCameraPreviewPlayer(EventSink));
+	}
+
+	return nullptr;
+}
+
+IMPLEMENT_MODULE(FMagicLeapCameraPreviewModule, MagicLeapCameraPreview);
diff --git a/Engine/Plugins/Lumin/MagicLeapCamera/Source/MagicLeapCameraPreview/Private/MagicLeapCameraPreviewModule.h b/Engine/Plugins/Lumin/MagicLeapCamera/Source/MagicLeapCameraPreview/Private/MagicLeapCameraPreviewModule.h
new file mode 100644
--- /dev/null
+++ b/Engine/Plugins/Lumin/MagicLeapCamera/Source/MagicLeapCameraPreview/Private/MagicLeapCameraPreviewModule.h
@@ -0,0 +1,21 @@
+// Copyright 1998-2019 Epic Games, Inc. All Rights Reserved.
+
+#pragma once
+
+#include "Modules/ModuleManager.h"
+#include "Containers/Ticker.h"
+#include "IMagicLeapCameraPreviewModule.h"
+
+class IMediaPlayer;
+
+DECLARE_LOG_CATEGORY_EXTERN(LogMagicLeapCameraPreview, Verbose, All);
+
+/**
+ * The public interface to this module.  In most cases, this interface is only public to sibling modules
+ * within this plugin.
+ */
+class FMagicLeapCameraPreviewModule : public IMagicLeapCameraPreviewModule
+{
+public:
+	TSharedPtr<IMediaPlayer, ESPMode::ThreadSafe> CreatePreviewPlayer(class IMediaEventSink& EventSink) override;
+};
diff --git a/Engine/Plugins/Lumin/MagicLeapCamera/Source/MagicLeapCameraPreview/Private/MagicLeapCameraPreviewPlayer.cpp b/Engine/Plugins/Lumin/MagicLeapCamera/Source/MagicLeapCameraPreview/Private/MagicLeapCameraPreviewPlayer.cpp
new file mode 100644
--- /dev/null
+++ b/Engine/Plugins/Lumin/MagicLeapCamera/Source/MagicLeapCameraPreview/Private/MagicLeapCameraPreviewPlayer.cpp
@@ -0,0 +1,632 @@
+// Copyright 1998-2019 Epic Games, Inc. All Rights Reserved.
+#include "MagicLeapCameraPreviewPlayer.h"
+#include "IMagicLeapCameraPlugin.h"
+#include "MagicLeapCameraPreviewModule.h"
+#include "Misc/Paths.h"
+#include "Lumin/LuminPlatformFile.h"
+#include "Lumin/LuminPlatformMisc.h"
+#include "IMediaEventSink.h"
+#include "MediaSamples.h"
+#include "Misc/CoreDelegates.h"
+#include "UObject/Class.h"
+#include "UObject/UObjectGlobals.h"
+#include "ExternalTexture.h"
+#include "ExternalTextureRendererGL.h"
+#include "MagicLeapHelperVulkan.h"
+
+struct FMagicLeapVideoTextureData
+{
+public:
+	FMagicLeapVideoTextureData()
+		: bIsVideoTextureValid(false)
+		, PreviousNativeBuffer(ML_INVALID_HANDLE)
+	{}
+
+	FTextureRHIRef VideoTexture;
+	bool bIsVideoTextureValid;
+	MLHandle PreviousNativeBuffer;
+};
+
+struct FMagicLeapVideoTextureDataVK : public FMagicLeapVideoTextureData
+{
+	FSamplerStateRHIRef VideoSampler;
+	TMap<uint64, FTextureRHIRef> VideoTexturePool;
+};
+
+struct FMagicLeapVideoTextureDataGL : public FMagicLeapVideoTextureData
+{
+public:
+
+	FMagicLeapVideoTextureDataGL()
+	{}
+
+	FExternalTextureRendererGL TextureRenderer;
+};
+
+FMagicLeapCameraPreviewPlayer::FMagicLeapCameraPreviewPlayer(IMediaEventSink& InEventSink)
+	: bMediaPrepared(false)
+	, CurrentState(EMediaState::Closed)
+	, EventSink(InEventSink)
+	, Samples(MakeShared<FMediaSamples, ESPMode::ThreadSafe>())
+	, bWasMediaPlayingBeforeAppPause(false)
+	, bPlaybackCompleted(false)
+{
+	if (FLuminPlatformMisc::ShouldUseVulkan())
+	{
+		TextureData = MakeShared<FMagicLeapVideoTextureDataVK, ESPMode::ThreadSafe>();
+	}
+	else
+	{
+		TextureData = MakeShared<FMagicLeapVideoTextureDataGL, ESPMode::ThreadSafe>();
+	}
+
+	CurrentState = Samples.IsValid() ? EMediaState::Closed : EMediaState::Error;
+}
+
+FMagicLeapCameraPreviewPlayer::~FMagicLeapCameraPreviewPlayer()
+{
+	Close();
+
+	MLHandle PreviewHandle = IMagicLeapCameraPlugin::Get().GetPreviewHandle();
+	if (MLHandleIsValid(PreviewHandle) && GSupportsImageExternal)
+	{
+		if (FLuminPlatformMisc::ShouldUseVulkan())
+		{
+			// Unregister the external texture on render thread
+			struct FReleaseVideoResourcesParams
+			{
+				FGuid PlayerGuid;
+			};
+
+			FReleaseVideoResourcesParams Params = { PlayerGuid };
+
+			ENQUEUE_RENDER_COMMAND(MagicLeapMediaPlayerDestroy)(
+				[Params](FRHICommandListImmediate& RHICmdList)
+				{
+					FExternalTextureRegistry::Get().UnregisterExternalTexture(Params.PlayerGuid);
+				});
+		}
+		else
+		{
+			// Unregister the external texture on render thread
+			struct FReleaseVideoResourcesParams
+			{
+				// Can we make this a TWeakPtr? We need to ensure that FMagicLeapMediaPlayer_::TextureData is not destroyed before
+				// we unregister the external texture and releae the VideoTexture.
+				TSharedPtr<FMagicLeapVideoTextureDataGL, ESPMode::ThreadSafe> TextureData;
+				FGuid PlayerGuid;
+			};
+
+			FReleaseVideoResourcesParams Params = { StaticCastSharedPtr<FMagicLeapVideoTextureDataGL>(TextureData), PlayerGuid };
+
+			ENQUEUE_RENDER_COMMAND(MagicLeapMediaPlayerDestroy)(
+				[Params](FRHICommandListImmediate& RHICmdList)
+				{
+					FExternalTextureRegistry::Get().UnregisterExternalTexture(Params.PlayerGuid);
+					// @todo: this causes a crash
+					//Params.TextureData->VideoTexture->Release();
+					Params.TextureData->TextureRenderer.DestroyImageKHR();
+				});
+		}
+
+		FlushRenderingCommands();
+	}
+}
+
+void FMagicLeapCameraPreviewPlayer::Close()
+{
+	if (CurrentState == EMediaState::Closed || CurrentState == EMediaState::Error)
+	{
+		return;
+	}
+
+	{
+		FScopeLock Lock(&CriticalSection);
+		bPlaybackCompleted = true;
+	}
+
+	// remove delegates if registered
+	if (ResumeHandle.IsValid())
+	{
+		FCoreDelegates::ApplicationHasEnteredForegroundDelegate.Remove(ResumeHandle);
+		ResumeHandle.Reset();
+	}
+
+	if (PauseHandle.IsValid())
+	{
+		FCoreDelegates::ApplicationWillEnterBackgroundDelegate.Remove(PauseHandle);
+		PauseHandle.Reset();
+	}
+
+	FMagicLeapCameraDisconnect DummyDisconnectCallback;
+	IMagicLeapCameraPlugin::Get().CameraDisconnect(DummyDisconnectCallback);
+
+	CurrentState = EMediaState::Closed;
+
+	bMediaPrepared = false;
+	TrackInfo.Empty();
+	SelectedTrack.Empty();
+
+	// notify listeners
+	EventSink.ReceiveMediaEvent(EMediaEvent::TracksChanged);
+	EventSink.ReceiveMediaEvent(EMediaEvent::MediaClosed);
+}
+
+IMediaCache& FMagicLeapCameraPreviewPlayer::GetCache()
+{
+	return *this;
+}
+
+IMediaControls& FMagicLeapCameraPreviewPlayer::GetControls()
+{
+	return *this;
+}
+
+FString FMagicLeapCameraPreviewPlayer::GetInfo() const
+{
+	return FString();
+}
+
+FName FMagicLeapCameraPreviewPlayer::GetPlayerName() const
+{
+	static FName PlayerName(TEXT("MagicLeapCameraPreview"));
+	return PlayerName;
+}
+
+IMediaSamples& FMagicLeapCameraPreviewPlayer::GetSamples()
+{
+	return *Samples.Get();
+}
+
+FString FMagicLeapCameraPreviewPlayer::GetStats() const
+{
+	return TEXT("MagicLeapCameraPreview stats information not implemented yet");
+}
+
+IMediaTracks& FMagicLeapCameraPreviewPlayer::GetTracks()
+{
+	return *this;
+}
+
+FString FMagicLeapCameraPreviewPlayer::GetUrl() const
+{
+	return FString();
+}
+
+IMediaView& FMagicLeapCameraPreviewPlayer::GetView()
+{
+	return *this;
+}
+
+bool FMagicLeapCameraPreviewPlayer::Open(const FString& Url, const IMediaOptions* Options)
+{
+	if (CurrentState == EMediaState::Error)
+	{
+		EventSink.ReceiveMediaEvent(EMediaEvent::MediaOpenFailed);
+		return false;
+	}
+
+	Close();
+
+	FMagicLeapCameraConnect DummyConnectCallback;
+	IMagicLeapCameraPlugin::Get().CameraConnect(DummyConnectCallback);
+
+	EventSink.ReceiveMediaEvent(EMediaEvent::MediaConnecting);
+
+	CurrentState = EMediaState::Preparing;
+
+	return true;
+}
+
+bool FMagicLeapCameraPreviewPlayer::Open(const TSharedRef<FArchive, ESPMode::ThreadSafe>& Archive, const FString& OriginalUrl, const IMediaOptions* Options)
+{
+	return false;
+}
+
+bool FMagicLeapCameraPreviewPlayer::CanControl(EMediaControl Control) const
+{
+	if (Control == EMediaControl::Pause)
+	{
+		return (CurrentState == EMediaState::Playing);
+	}
+
+	if (Control == EMediaControl::Resume)
+	{
+		return (CurrentState == EMediaState::Paused);
+	}
+
+	return false;
+}
+
+FTimespan FMagicLeapCameraPreviewPlayer::GetDuration() const
+{
+	return FTimespan::Zero();
+}
+
+float FMagicLeapCameraPreviewPlayer::GetRate() const
+{
+	return (CurrentState == EMediaState::Playing) ? 1.0f : 0.0f;
+}
+
+EMediaState FMagicLeapCameraPreviewPlayer::GetState() const
+{
+	return CurrentState;
+}
+
+EMediaStatus FMagicLeapCameraPreviewPlayer::GetStatus() const
+{
+	return CurrentState == EMediaState::Preparing ? EMediaStatus::Connecting : EMediaStatus::None;
+}
+
+TRangeSet<float> FMagicLeapCameraPreviewPlayer::GetSupportedRates(EMediaRateThinning Thinning) const
+{
+	TRangeSet<float> Result;
+
+	Result.Add(TRange<float>(0.0f));
+	Result.Add(TRange<float>(1.0f));
+
+	return Result;
+}
+
+FTimespan FMagicLeapCameraPreviewPlayer::GetTime() const
+{
+	return FTimespan::Zero();
+}
+
+bool FMagicLeapCameraPreviewPlayer::IsLooping() const
+{
+	return false;
+}
+
+bool FMagicLeapCameraPreviewPlayer::Seek(const FTimespan& Time)
+{
+	return false;
+}
+
+bool FMagicLeapCameraPreviewPlayer::SetLooping(bool Looping)
+{
+	return false;
+}
+
+bool FMagicLeapCameraPreviewPlayer::SetRate(float Rate)
+{
+	if ((CurrentState == EMediaState::Closed) || (CurrentState == EMediaState::Error) || (CurrentState == EMediaState::Preparing))
+	{
+		UE_LOG(LogMagicLeapCameraPreview, Warning, TEXT("Cannot set rate while closed, preparing, or in error state"));
+		return false;
+	}
+
+	if (Rate == GetRate())
+	{
+		// rate already set
+		return true;
+	}
+
+	bool bResult = false;
+	if (Rate == 0.0f)
+	{
+		CurrentState = EMediaState::Paused;
+		EventSink.ReceiveMediaEvent(EMediaEvent::PlaybackSuspended);
+	}
+	else if (Rate == 1.0f)
+	{
+		CurrentState = EMediaState::Playing;
+		EventSink.ReceiveMediaEvent(EMediaEvent::PlaybackResumed);
+		bResult = true;
+	}
+	else
+	{
+		UE_LOG(LogMagicLeapCameraPreview, Error, TEXT("Rate %f not supported by MagicLeapMedia."), Rate);
+		bResult = false;
+	}
+
+	return bResult;
+}
+
+bool FMagicLeapCameraPreviewPlayer::SetNativeVolume(float Volume)
+{
+	return false;
+}
+
+void FMagicLeapCameraPreviewPlayer::SetGuid(const FGuid& Guid)
+{
+	PlayerGuid = Guid;
+}
+
+void FMagicLeapCameraPreviewPlayer::TickFetch(FTimespan DeltaTime, FTimespan Timecode)
+{
+	if (CurrentState != EMediaState::Playing && CurrentState != EMediaState::Paused)
+	{
+		return;
+	}
+
+	if (GSupportsImageExternal)
+	{
+		MLHandle PreviewHandle = IMagicLeapCameraPlugin::Get().GetPreviewHandle();
+		if (!MLHandleIsValid(PreviewHandle))
+		{
+			return;
+		}
+
+		if (FLuminPlatformMisc::ShouldUseVulkan())
+		{
+			struct FWriteVideoSampleParams
+			{
+				FMagicLeapCameraPreviewPlayer* MediaPlayer;
+				TWeakPtr<FMagicLeapVideoTextureDataVK, ESPMode::ThreadSafe> TextureData;
+				FGuid PlayerGuid;
+				MLHandle PreviewHandle;
+			};
+
+			FWriteVideoSampleParams Params = { this, StaticCastSharedPtr<FMagicLeapVideoTextureDataVK>(TextureData), PlayerGuid, PreviewHandle };
+
+			ENQUEUE_RENDER_COMMAND(MagicLeapMediaPlayerWriteVideoSample)(
+				[Params](FRHICommandListImmediate& RHICmdList)
+				{
+					auto TextureDataPtr = Params.TextureData.Pin();
+
+					if (TextureDataPtr->PreviousNativeBuffer != 0 && MLHandleIsValid(TextureDataPtr->PreviousNativeBuffer))
+					{
+						TextureDataPtr->PreviousNativeBuffer = 0;
+					}
+
+					MLHandle NativeBuffer = Params.PreviewHandle;
+
+					{
+						FScopeLock Lock(&Params.MediaPlayer->CriticalSection);
+						if (Params.MediaPlayer->bPlaybackCompleted)
+						{
+							TextureDataPtr->VideoTexturePool.Empty();
+							Params.MediaPlayer->bPlaybackCompleted = false;
+						}
+					}
+
+					if (!TextureDataPtr->VideoTexturePool.Contains((uint64)NativeBuffer))
+					{
+						FTextureRHIRef NewMediaTexture;
+						if (!FMagicLeapHelperVulkan::GetMediaTexture(NewMediaTexture, TextureDataPtr->VideoSampler, NativeBuffer))
+						{
+							UE_LOG(LogMagicLeapCameraPreview, Error, TEXT("Failed to get next media texture."));
+							return;
+						}
+
+						TextureDataPtr->VideoTexturePool.Add((uint64)NativeBuffer, NewMediaTexture);
+
+						if (TextureDataPtr->VideoTexture == nullptr)
+						{
+							FRHIResourceCreateInfo CreateInfo;
+							TextureDataPtr->VideoTexture = RHICmdList.CreateTextureExternal2D(1, 1, PF_R8G8B8A8, 1, 1, 0, CreateInfo);
+						}
+
+						FMagicLeapHelperVulkan::AliasMediaTexture(TextureDataPtr->VideoTexture, NewMediaTexture);
+					}
+					else
+					{
+						FTextureRHIRef* const PooledMediaTexture = TextureDataPtr->VideoTexturePool.Find((uint64)NativeBuffer);
+						check(PooledMediaTexture != nullptr);
+						FMagicLeapHelperVulkan::AliasMediaTexture(TextureDataPtr->VideoTexture, *PooledMediaTexture);
+					}
+
+					if (!TextureDataPtr->bIsVideoTextureValid)
+					{
+						Params.MediaPlayer->RegisterExternalTexture(Params.PlayerGuid, TextureDataPtr->VideoTexture, TextureDataPtr->VideoSampler);
+						TextureDataPtr->bIsVideoTextureValid = true;
+					}
+
+					TextureDataPtr->PreviousNativeBuffer = NativeBuffer;
+				});
+		}
+		else
+		{
+			struct FWriteVideoSampleParams
+			{
+				FMagicLeapCameraPreviewPlayer* MediaPlayer;
+				TWeakPtr<FMagicLeapVideoTextureDataGL, ESPMode::ThreadSafe> TextureData;
+				FGuid PlayerGuid;
+				MLHandle PreviewHandle;
+			};
+
+			FWriteVideoSampleParams Params = { this, StaticCastSharedPtr<FMagicLeapVideoTextureDataGL>(TextureData), PlayerGuid, PreviewHandle };
+
+			ENQUEUE_RENDER_COMMAND(MagicLeapMediaPlayerWriteVideoSample)(
+				[Params](FRHICommandListImmediate& RHICmdList)
+				{
+					auto TextureDataPtr = Params.TextureData.Pin();
+					FTextureRHIRef MediaVideoTexture = TextureDataPtr->VideoTexture;
+					if (MediaVideoTexture == nullptr)
+					{
+						FRHIResourceCreateInfo CreateInfo;
+						MediaVideoTexture = RHICmdList.CreateTextureExternal2D(1, 1, PF_R8G8B8A8, 1, 1, 0, CreateInfo);
+						TextureDataPtr->VideoTexture = MediaVideoTexture;
+
+						if (MediaVideoTexture == nullptr)
+						{
+							UE_LOG(LogMagicLeapCameraPreview, Warning, TEXT("CreateTextureExternal2D failed!"));
+							return;
+						}
+
+						TextureDataPtr->bIsVideoTextureValid = false;
+					}
+
+					// MLHandle because Unreal's uint64 is 'unsigned long long *' whereas uint64_t for the C-API is 'unsigned long *'
+					// TODO: Fix the Unreal types for the above comment.
+					MLHandle NativeBuffer = Params.PreviewHandle;
+					int32 TextureID = *reinterpret_cast<int32*>(MediaVideoTexture->GetNativeResource());
+					TextureDataPtr->TextureRenderer.DestroyImageKHR();
+					TextureDataPtr->PreviousNativeBuffer = NativeBuffer;
+
+					if (TextureDataPtr->TextureRenderer.CreateImageKHR(reinterpret_cast<void*>(NativeBuffer)))
+					{
+						TextureDataPtr->TextureRenderer.BindImageKHRToTexture(TextureID);
+						if (!TextureDataPtr->bIsVideoTextureValid)
+						{
+							FSamplerStateInitializerRHI SamplerStateInitializer(SF_Bilinear, AM_Clamp, AM_Clamp, AM_Clamp);
+							FSamplerStateRHIRef SamplerStateRHI = RHICreateSamplerState(SamplerStateInitializer);
+							Params.MediaPlayer->RegisterExternalTexture(Params.PlayerGuid, MediaVideoTexture, SamplerStateRHI);
+							TextureDataPtr->bIsVideoTextureValid = true;
+						}
+					}
+				});
+		}
+	}
+}
+
+void FMagicLeapCameraPreviewPlayer::RegisterExternalTexture(const FGuid& InGuid, FTextureRHIRef& InTextureRHI, FSamplerStateRHIRef& InSamplerStateRHI)
+{
+	FExternalTextureRegistry::Get().RegisterExternalTexture(InGuid, InTextureRHI, InSamplerStateRHI, FLinearColor(1.0f, 0.0f, 0.0f, 1.0f), FLinearColor(0.0f, 0.0f, 0.0f, 0.0f));
+}
+
+void FMagicLeapCameraPreviewPlayer::TickInput(FTimespan DeltaTime, FTimespan Timecode)
+{
+	if (!bMediaPrepared)
+	{
+		bMediaPrepared = MLHandleIsValid(IMagicLeapCameraPlugin::Get().GetPreviewHandle());
+
+		if (bMediaPrepared)
+		{
+			CurrentState = EMediaState::Stopped;
+
+			TrackInfo.Add(EMediaTrackType::Video);
+			SelectedTrack.Add(EMediaTrackType::Video, INDEX_NONE);
+
+			TrackInfo[EMediaTrackType::Video].Add(0);
+			SelectedTrack[EMediaTrackType::Video] = 0;
+
+			// notify listeners
+			EventSink.ReceiveMediaEvent(EMediaEvent::TracksChanged);
+			EventSink.ReceiveMediaEvent(EMediaEvent::MediaOpened);
+		}
+		else
+		{
+			return;
+		}
+	}
+
+	if (CurrentState != EMediaState::Playing)
+	{
+		// remove delegates if registered
+		if (ResumeHandle.IsValid())
+		{
+			FCoreDelegates::ApplicationHasEnteredForegroundDelegate.Remove(ResumeHandle);
+			ResumeHandle.Reset();
+		}
+
+		if (PauseHandle.IsValid())
+		{
+			FCoreDelegates::ApplicationWillEnterBackgroundDelegate.Remove(PauseHandle);
+			PauseHandle.Reset();
+		}
+
+		return;
+	}
+
+	// register delegate if not registered
+	if (!ResumeHandle.IsValid())
+	{
+		ResumeHandle = FCoreDelegates::ApplicationHasEnteredForegroundDelegate.AddRaw(this, &FMagicLeapCameraPreviewPlayer::HandleApplicationHasEnteredForeground);
+	}
+	if (!PauseHandle.IsValid())
+	{
+		PauseHandle = FCoreDelegates::ApplicationWillEnterBackgroundDelegate.AddRaw(this, &FMagicLeapCameraPreviewPlayer::HandleApplicationWillEnterBackground);
+	}
+}
+
+bool FMagicLeapCameraPreviewPlayer::GetAudioTrackFormat(int32 TrackIndex, int32 FormatIndex, FMediaAudioTrackFormat& OutFormat) const
+{
+	return false;
+}
+
+int32 FMagicLeapCameraPreviewPlayer::GetNumTracks(EMediaTrackType TrackType) const
+{
+	if (TrackInfo.Contains(TrackType))
+	{
+		return TrackInfo[TrackType].Num();
+	}
+
+	return 0;
+}
+
+int32 FMagicLeapCameraPreviewPlayer::GetNumTrackFormats(EMediaTrackType TrackType, int32 TrackIndex) const
+{
+	return ((TrackIndex >= 0) && (TrackIndex < GetNumTracks(TrackType))) ? 1 : 0;
+}
+
+int32 FMagicLeapCameraPreviewPlayer::GetSelectedTrack(EMediaTrackType TrackType) const
+{
+	return TrackType == EMediaTrackType::Video ? 0 : INDEX_NONE;
+}
+
+FText FMagicLeapCameraPreviewPlayer::GetTrackDisplayName(EMediaTrackType TrackType, int32 TrackIndex) const
+{
+	return FText::GetEmpty();
+}
+
+int32 FMagicLeapCameraPreviewPlayer::GetTrackFormat(EMediaTrackType TrackType, int32 TrackIndex) const
+{
+	return (GetSelectedTrack(TrackType) != INDEX_NONE) ? 0 : INDEX_NONE;
+}
+
+FString FMagicLeapCameraPreviewPlayer::GetTrackLanguage(EMediaTrackType TrackType, int32 TrackIndex) const
+{
+	return FString();
+}
+
+FString FMagicLeapCameraPreviewPlayer::GetTrackName(EMediaTrackType TrackType, int32 TrackIndex) const
+{
+	return FString();
+}
+
+bool FMagicLeapCameraPreviewPlayer::GetVideoTrackFormat(int32 TrackIndex, int32 FormatIndex, FMediaVideoTrackFormat& OutFormat) const
+{
+	(void)TrackIndex;
+	(void)FormatIndex;
+	(void)OutFormat;
+	// TODO: Get these dimensions from the native buffer itself.
+	OutFormat.Dim = FIntPoint(512, 512);
+	// TODO: Don't hardcode. Get from C-API. ml_media_player api does not provide that right now. Try the ml_media_codec api.
+	OutFormat.FrameRate = 30.0f;
+	OutFormat.FrameRates = TRange<float>(30.0f);
+	OutFormat.TypeName = TEXT("BGRA");
+
+	return true;
+}
+
+bool FMagicLeapCameraPreviewPlayer::SelectTrack(EMediaTrackType TrackType, int32 TrackIndex)
+{
+	if (TrackInfo.Contains(TrackType) && CurrentState != EMediaState::Preparing)
+	{
+		if (TrackInfo[TrackType].IsValidIndex(TrackIndex))
+		{
+			SelectedTrack[TrackType] = TrackIndex;
+			return true;
+		}
+	}
+
+	return false;
+}
+
+bool FMagicLeapCameraPreviewPlayer::SetTrackFormat(EMediaTrackType TrackType, int32 TrackIndex, int32 FormatIndex)
+{
+	return false;
+}
+
+void FMagicLeapCameraPreviewPlayer::HandleApplicationHasEnteredForeground()
+{
+	// check state in case changed before ticked
+	if (CurrentState == EMediaState::Paused && bWasMediaPlayingBeforeAppPause)
+	{
+		// Unpause
+		SetRate(1.0f);
+	}
+}
+
+void FMagicLeapCameraPreviewPlayer::HandleApplicationWillEnterBackground()
+{
+	bWasMediaPlayingBeforeAppPause = (CurrentState == EMediaState::Playing);
+	// check state in case changed before ticked
+	if (bWasMediaPlayingBeforeAppPause)
+	{
+		// pause
+		SetRate(0.0f);
+	}
+}
