UE-81320 //UE4/Dev-VR - Check Balanced Macros - PRAGMA_DISABLE_OPTIMIZATION - 1 warning
-Removed #ifdef 0'ed out pragma

diff --git a/Engine/Plugins/Lumin/MagicLeapCamera/Source/MagicLeapCameraPreview/Private/MagicLeapCameraPreviewModule.cpp b/Engine/Plugins/Lumin/MagicLeapCamera/Source/MagicLeapCameraPreview/Private/MagicLeapCameraPreviewModule.cpp
--- a/Engine/Plugins/Lumin/MagicLeapCamera/Source/MagicLeapCameraPreview/Private/MagicLeapCameraPreviewModule.cpp
+++ b/Engine/Plugins/Lumin/MagicLeapCamera/Source/MagicLeapCameraPreview/Private/MagicLeapCameraPreviewModule.cpp
@@ -4,10 +4,6 @@
 #include "MagicLeapCameraPreviewPlayer.h"
 #include "Async/Async.h"
 
-#if 0
-PRAGMA_DISABLE_OPTIMIZATION
-#endif
-
 DEFINE_LOG_CATEGORY(LogMagicLeapCameraPreview);
 
 TSharedPtr<IMediaPlayer, ESPMode::ThreadSafe> FMagicLeapCameraPreviewModule::CreatePreviewPlayer(class IMediaEventSink& EventSink)
