Build fix

diff --git a/Engine/Plugins/Lumin/MagicLeapCamera/Source/MagicLeapCameraPreview/Private/MagicLeapCameraPreviewPlayer.cpp b/Engine/Plugins/Lumin/MagicLeapCamera/Source/MagicLeapCameraPreview/Private/MagicLeapCameraPreviewPlayer.cpp
--- a/Engine/Plugins/Lumin/MagicLeapCamera/Source/MagicLeapCameraPreview/Private/MagicLeapCameraPreviewPlayer.cpp
+++ b/Engine/Plugins/Lumin/MagicLeapCamera/Source/MagicLeapCameraPreview/Private/MagicLeapCameraPreviewPlayer.cpp
@@ -395,7 +395,7 @@ void FMagicLeapCameraPreviewPlayer::TickFetch(FTimespan DeltaTime, FTimespan Tim
 
 						if (TextureDataPtr->VideoTexture == nullptr)
 						{
-							FRHIResourceCreateInfo CreateInfo;
+							FRHIResourceCreateInfo CreateInfo(TEXT("FMagicLeapCameraPreviewPlayer_VideoTexture"));
 							TextureDataPtr->VideoTexture = RHICreateTextureExternal2D(1, 1, PF_R8G8B8A8, 1, 1, TexCreate_None, CreateInfo);
 						}
 
@@ -436,7 +436,7 @@ void FMagicLeapCameraPreviewPlayer::TickFetch(FTimespan DeltaTime, FTimespan Tim
 					FTextureRHIRef MediaVideoTexture = TextureDataPtr->VideoTexture;
 					if (MediaVideoTexture == nullptr)
 					{
-						FRHIResourceCreateInfo CreateInfo;
+						FRHIResourceCreateInfo CreateInfo(TEXT("MediaVideoTexture"));
 						MediaVideoTexture = RHICreateTextureExternal2D(1, 1, PF_R8G8B8A8, 1, 1, TexCreate_None, CreateInfo);
 						TextureDataPtr->VideoTexture = MediaVideoTexture;
 
diff --git a/Engine/Plugins/Lumin/MagicLeapMedia/Source/MagicLeapMedia/Private/MagicLeapMediaPlayer.cpp b/Engine/Plugins/Lumin/MagicLeapMedia/Source/MagicLeapMedia/Private/MagicLeapMediaPlayer.cpp
--- a/Engine/Plugins/Lumin/MagicLeapMedia/Source/MagicLeapMedia/Private/MagicLeapMediaPlayer.cpp
+++ b/Engine/Plugins/Lumin/MagicLeapMedia/Source/MagicLeapMedia/Private/MagicLeapMediaPlayer.cpp
@@ -665,7 +665,7 @@ void FMagicLeapMediaPlayer::TickFetch(FTimespan DeltaTime, FTimespan Timecode)
 						
 						if (TextureDataPtr->VideoTexture == nullptr)
 						{
-							FRHIResourceCreateInfo CreateInfo;
+							FRHIResourceCreateInfo CreateInfo(TEXT("FMagicLeapMediaPlayer_VideoTexture"));
 							TextureDataPtr->VideoTexture = RHICreateTextureExternal2D(1, 1, PF_R8G8B8A8, 1, 1, TexCreate_None, CreateInfo);
 						}
 
@@ -712,7 +712,7 @@ void FMagicLeapMediaPlayer::TickFetch(FTimespan DeltaTime, FTimespan Timecode)
 					FTextureRHIRef MediaVideoTexture = TextureDataPtr->VideoTexture;
 					if (MediaVideoTexture == nullptr)
 					{
-						FRHIResourceCreateInfo CreateInfo;
+						FRHIResourceCreateInfo CreateInfo(TEXT("FMagicLeapMediaPlayer_MediaVideoTexture"));
 						MediaVideoTexture = RHICreateTextureExternal2D(1, 1, PF_R8G8B8A8, 1, 1, TexCreate_None, CreateInfo);
 						TextureDataPtr->VideoTexture = MediaVideoTexture;
 
diff --git a/Engine/Plugins/Lumin/MagicLeapMedia/Source/MagicLeapMediaCodec/Private/MagicLeapMediaCodecPlayer.cpp b/Engine/Plugins/Lumin/MagicLeapMedia/Source/MagicLeapMediaCodec/Private/MagicLeapMediaCodecPlayer.cpp
--- a/Engine/Plugins/Lumin/MagicLeapMedia/Source/MagicLeapMediaCodec/Private/MagicLeapMediaCodecPlayer.cpp
+++ b/Engine/Plugins/Lumin/MagicLeapMedia/Source/MagicLeapMediaCodec/Private/MagicLeapMediaCodecPlayer.cpp
@@ -762,7 +762,7 @@ void FMagicLeapMediaCodecPlayer::TickFetch(FTimespan DeltaTime, FTimespan Timeco
 
 						if (TextureDataPtr->VideoTexture == nullptr)
 						{
-							FRHIResourceCreateInfo CreateInfo;
+							FRHIResourceCreateInfo CreateInfo(TEXT("FMagicLeapMediaCodecPlayer_VideoTexture"));
 							TextureDataPtr->VideoTexture = RHICreateTextureExternal2D(1, 1, PF_R8G8B8A8, 1, 1, TexCreate_None, CreateInfo);
 						}
 
@@ -813,7 +813,7 @@ void FMagicLeapMediaCodecPlayer::TickFetch(FTimespan DeltaTime, FTimespan Timeco
 					FTextureRHIRef MediaVideoTexture = TextureDataPtr->VideoTexture;
 					if (MediaVideoTexture == nullptr)
 					{
-						FRHIResourceCreateInfo CreateInfo;
+						FRHIResourceCreateInfo CreateInfo(TEXT("MediaVideoTexture"));
 						MediaVideoTexture = RHICreateTextureExternal2D(1, 1, PF_R8G8B8A8, 1, 1, TexCreate_None, CreateInfo);
 						TextureDataPtr->VideoTexture = MediaVideoTexture;
 
diff --git a/Engine/Plugins/Media/ElectraUtil/Source/ElectraSamples/Public/GStreamer/ElectraTextureSample.h b/Engine/Plugins/Media/ElectraUtil/Source/ElectraSamples/Public/GStreamer/ElectraTextureSample.h
--- a/Engine/Plugins/Media/ElectraUtil/Source/ElectraSamples/Public/GStreamer/ElectraTextureSample.h
+++ b/Engine/Plugins/Media/ElectraUtil/Source/ElectraSamples/Public/GStreamer/ElectraTextureSample.h
@@ -37,7 +37,7 @@ public:
 		const uint32 CreateFlags = TexCreate_Dynamic | TexCreate_SRGB;
 
 		TRefCountPtr<FRHITexture2D> DummyTexture2DRHI;
-		FRHIResourceCreateInfo CreateInfo;
+		FRHIResourceCreateInfo CreateInfo(TEXT("DummyTexture2D"));
 
 		RHICreateTargetableShaderResource2D(
 			TotalSize.X,
diff --git a/Engine/Plugins/Runtime/AR/AppleAR/AppleARKit/Source/AppleARKit/Private/AppleARKitTextures.cpp b/Engine/Plugins/Runtime/AR/AppleAR/AppleARKit/Source/AppleARKit/Private/AppleARKitTextures.cpp
--- a/Engine/Plugins/Runtime/AR/AppleAR/AppleARKit/Source/AppleARKit/Private/AppleARKitTextures.cpp
+++ b/Engine/Plugins/Runtime/AR/AppleAR/AppleARKit/Source/AppleARKit/Private/AppleARKitTextures.cpp
@@ -981,7 +981,7 @@ public:
 		if (!DecodedTextureRef || Size != OutputSize)
 		{
 			Size = OutputSize;
-			FRHIResourceCreateInfo CreateInfo;
+			FRHIResourceCreateInfo CreateInfo(TEXT("FARKitCameraVideoResource_DecodedTexture"));
 			{
 				FScopeLock ScopeLock(&DecodedTextureLock);
 				DecodedTextureRef = RHICreateTexture2D(Size.X, Size.Y, PF_B8G8R8A8, 1, 1, TexCreate_Dynamic | TexCreate_ShaderResource | TexCreate_UAV, CreateInfo);
diff --git a/Engine/Source/Runtime/WebBrowser/Private/IOS/IOSPlatformWebBrowser.cpp b/Engine/Source/Runtime/WebBrowser/Private/IOS/IOSPlatformWebBrowser.cpp
--- a/Engine/Source/Runtime/WebBrowser/Private/IOS/IOSPlatformWebBrowser.cpp
+++ b/Engine/Source/Runtime/WebBrowser/Private/IOS/IOSPlatformWebBrowser.cpp
@@ -158,7 +158,7 @@ class SIOSWebBrowserWidget : public SLeafWidget
 							FTextureRHIRef VideoTexture = [NativeWebBrowser GetVideoTexture];
 							if (VideoTexture == nullptr)
 							{
-								FRHIResourceCreateInfo CreateInfo;
+								FRHIResourceCreateInfo CreateInfo(TEXT("SIOSWebBrowserWidget_VideoTexture"));
 								FIntPoint Size = Params.Size;
 								VideoTexture = RHICreateTextureExternal2D(Size.X, Size.Y, PF_R8G8B8A8, 1, 1, TexCreate_None, CreateInfo);
 								[NativeWebBrowser SetVideoTexture : VideoTexture];
